AnimDatabase: Added support for SavGol filter

#Synced-CL 50709564

diff --git a/Engine/Plugins/Experimental/Animation/AnimDatabase/Source/AnimDatabase/Private/AnimDatabase.cpp b/Engine/Plugins/Experimental/Animation/AnimDatabase/Source/AnimDatabase/Private/AnimDatabase.cpp
--- a/Engine/Plugins/Experimental/Animation/AnimDatabase/Source/AnimDatabase/Private/AnimDatabase.cpp
+++ b/Engine/Plugins/Experimental/Animation/AnimDatabase/Source/AnimDatabase/Private/AnimDatabase.cpp
@@ -6045,15 +6045,32 @@ void UAnimDatabaseFunction_UpdateIKBones::Run_Implementation(UAnimDatabase* InDa
 
 void UAnimDatabaseFunction_SetRootFromLookAt::Run_Implementation(UAnimDatabase* InDatabase, const FAnimDatabaseFrameRanges& InFrameRanges, const bool bShouldTransact)
 {
-	const FAnimDatabaseFrameAttribute SmoothedRootTransform = UAnimDatabaseFrameAttributeLibrary::FrameAttributeSmoothedRootTransform(
-		InDatabase,
-		InFrameRanges,
-		RootLocationBone,
-		RootDirectionBone,
-		RootDirectionBoneDirection,
-		RootDirection,
-		LocationSmoothingTime,
-		DirectionSmoothingTime);
+
+
+	const FAnimDatabaseFrameAttribute SmoothedRootTransform = bUseSavGolFilter ?
+		
+		UAnimDatabaseFrameAttributeLibrary::FrameAttributeSavGolSmoothedRootTransform(
+			InDatabase,
+			InFrameRanges,
+			RootLocationBone,
+			RootDirectionBone,
+			RootDirectionBoneDirection,
+			RootDirection,
+			FMath::RoundToInt(LocationSavGolFilterWidth * InDatabase->GetFrameRate().AsDecimal()),
+			LocationSavGolPolynomialDegree,
+			FMath::RoundToInt(DirectionSavGolFilterWidth * InDatabase->GetFrameRate().AsDecimal()),
+			DirectionSavGolPolynomialDegree,
+			bUseGaussianWindowedSavGolFilter) :
+		
+		UAnimDatabaseFrameAttributeLibrary::FrameAttributeGaussianSmoothedRootTransform(
+			InDatabase,
+			InFrameRanges,
+			RootLocationBone,
+			RootDirectionBone,
+			RootDirectionBoneDirection,
+			RootDirection,
+			LocationSmoothingTime * InDatabase->GetFrameRate().AsDecimal(),
+			DirectionSmoothingTime * InDatabase->GetFrameRate().AsDecimal());
 
 	const FAnimDatabaseFrameAttribute PelvisTransform = UAnimDatabaseFrameAttributeLibrary::MakeBoneGlobalTransformFrameAttributeFromName(InDatabase, InFrameRanges, PelvisBone);
 
diff --git a/Engine/Plugins/Experimental/Animation/AnimDatabase/Source/AnimDatabase/Private/AnimDatabaseFrameAttribute.cpp b/Engine/Plugins/Experimental/Animation/AnimDatabase/Source/AnimDatabase/Private/AnimDatabaseFrameAttribute.cpp
--- a/Engine/Plugins/Experimental/Animation/AnimDatabase/Source/AnimDatabase/Private/AnimDatabaseFrameAttribute.cpp
+++ b/Engine/Plugins/Experimental/Animation/AnimDatabase/Source/AnimDatabase/Private/AnimDatabaseFrameAttribute.cpp
@@ -7608,6 +7608,56 @@ FAnimDatabaseFrameAttribute UAnimDatabaseFrameAttributeLibrary::FrameAttributeFi
 	}
 }
 
+FAnimDatabaseFrameAttribute UAnimDatabaseFrameAttributeLibrary::FrameAttributeFilterSavGol(const FAnimDatabaseFrameAttribute& A, const int32 FilterWidthInFrames, const int32 PolynomialDegree, const bool bGaussianWindowed)
+{
+	if (!A.IsValid())
+	{
+		UE_LOGFMT(LogAnimDatabase, Error, "FrameAttributeFilterSavGol: Invalid FrameAttribute.");
+		return FAnimDatabaseFrameAttribute();
+	}
+
+	if (A.Type != EAnimDatabaseAttributeType::Float && A.Type != EAnimDatabaseAttributeType::Location && A.Type != EAnimDatabaseAttributeType::Direction)
+	{
+		UE_LOGFMT(LogAnimDatabase, Error, "FrameAttributeFilterSavGol: FrameAttribute Types don't match. Got {Type}. Expected Float, Location, or Direction.", *UEnum::GetValueAsString(A.Type));
+		return FAnimDatabaseFrameAttribute();
+	}
+
+	if (FilterWidthInFrames < 0 || PolynomialDegree < 1)
+	{
+		UE_LOGFMT(LogAnimDatabase, Error, "FrameAttributeFilterSavGol: Invalid Settings.");
+		return FAnimDatabaseFrameAttribute();
+	}
+
+	if (A.Type == EAnimDatabaseAttributeType::Float)
+	{
+		FAnimDatabaseFrameAttribute Out = MakeEmptyFrameAttribute();
+		Out.Type = EAnimDatabaseAttributeType::Float;
+		UE::Learning::FrameAttribute::FilterSavGol(*Out.FrameAttribute, *A.FrameAttribute, FilterWidthInFrames, PolynomialDegree, bGaussianWindowed);
+		return Out;
+	}
+	else if (A.Type == EAnimDatabaseAttributeType::Location)
+	{
+		FAnimDatabaseFrameAttribute Out = MakeEmptyFrameAttribute();
+		Out.Type = EAnimDatabaseAttributeType::Location;
+		UE::Learning::FrameAttribute::FilterSavGol(*Out.FrameAttribute, *A.FrameAttribute, FilterWidthInFrames, PolynomialDegree, bGaussianWindowed);
+		return Out;
+	}
+	else if (A.Type == EAnimDatabaseAttributeType::Direction)
+	{
+		FAnimDatabaseFrameAttribute Out = MakeEmptyFrameAttribute();
+		Out.Type = EAnimDatabaseAttributeType::Direction;
+		UE::Learning::FFrameAttribute Tmp;
+		UE::Learning::FrameAttribute::FilterSavGol(Tmp, *A.FrameAttribute, FilterWidthInFrames, PolynomialDegree, bGaussianWindowed);
+		UE::Learning::FrameAttribute::Normalize(*Out.FrameAttribute, Tmp);
+		return Out;
+	}
+	else
+	{
+		checkNoEntry();
+		return FAnimDatabaseFrameAttribute();
+	}
+}
+
 FAnimDatabaseFrameAttribute UAnimDatabaseFrameAttributeLibrary::FrameAttributeFilterMajorityVote(const FAnimDatabaseFrameAttribute& A, const int32 FilterWidthInFrames)
 {
 	if (!A.IsValid())
@@ -8221,7 +8271,36 @@ FAnimDatabaseFrameRanges UAnimDatabaseFrameAttributeLibrary::FrameAttributeOutli
 	return UAnimDatabaseFrameRangesLibrary::FrameRangesPad(UAnimDatabaseFrameRangesLibrary::FrameRangesUnionFromArrayView(BoneFrameRanges), Padding / 2, Padding / 2);
 }
 
-FAnimDatabaseFrameAttribute UAnimDatabaseFrameAttributeLibrary::FrameAttributeSmoothedRootTransform(
+FAnimDatabaseFrameAttribute UAnimDatabaseFrameAttributeLibrary::FrameAttributeSavGolSmoothedRootTransform(
+	const UAnimDatabase* Database,
+	const FAnimDatabaseFrameRanges& FrameRanges,
+	const FName RootLocationProjectionBoneName,
+	const FName RootRotationProjectionBoneName,
+	const FVector RootRotationBoneForwardDirection,
+	const FVector RootForwardDirection,
+	const int32 LocationSavGolFrameNum,
+	const int32 LocationSavGolPolynomialDegree,
+	const int32 DirectionSavGolFrameNum,
+	const int32 DirectionSavGolPolynomialDegree,
+	const bool bGaussianWindowed)
+{
+	UE::Tasks::TTask<FAnimDatabaseFrameAttribute> RootLocationTask = UE::Tasks::Launch(UE_SOURCE_LOCATION, [Database, &FrameRanges, RootLocationProjectionBoneName, LocationSavGolFrameNum, LocationSavGolPolynomialDegree, bGaussianWindowed]()
+		{
+			const FAnimDatabaseFrameAttribute RootLocation = MakeBoneGlobalLocationFrameAttributeFromName(Database, FrameRanges, RootLocationProjectionBoneName);
+			return FrameAttributeFilterSavGol(FrameAttributeProject(RootLocation), LocationSavGolFrameNum, LocationSavGolPolynomialDegree, bGaussianWindowed);
+		});
+
+	UE::Tasks::TTask<FAnimDatabaseFrameAttribute> RootRotationTask = UE::Tasks::Launch(UE_SOURCE_LOCATION, [Database, &FrameRanges, RootRotationProjectionBoneName, RootRotationBoneForwardDirection, RootForwardDirection, DirectionSavGolFrameNum, DirectionSavGolPolynomialDegree, bGaussianWindowed]()
+		{
+			const FAnimDatabaseFrameAttribute RootDirection = MakeBoneGlobalDirectionFrameAttributeFromName(Database, FrameRanges, RootRotationProjectionBoneName, RootRotationBoneForwardDirection);
+			const FAnimDatabaseFrameAttribute SmoothedRootDirection = FrameAttributeFilterSavGol(FrameAttributeProject(RootDirection), DirectionSavGolFrameNum, DirectionSavGolPolynomialDegree, bGaussianWindowed);
+			return FrameAttributeInvert(FrameAttributeRotationBetweenDirection(SmoothedRootDirection, RootForwardDirection));
+		});
+
+	return MakeTransformFrameAttributeNoScale(RootLocationTask.GetResult(), RootRotationTask.GetResult());
+}
+
+FAnimDatabaseFrameAttribute UAnimDatabaseFrameAttributeLibrary::FrameAttributeGaussianSmoothedRootTransform(
 	const UAnimDatabase* Database,
 	const FAnimDatabaseFrameRanges& FrameRanges,
 	const FName RootLocationProjectionBoneName,
diff --git a/Engine/Plugins/Experimental/Animation/AnimDatabase/Source/AnimDatabase/Public/AnimDatabase.h b/Engine/Plugins/Experimental/Animation/AnimDatabase/Source/AnimDatabase/Public/AnimDatabase.h
--- a/Engine/Plugins/Experimental/Animation/AnimDatabase/Source/AnimDatabase/Public/AnimDatabase.h
+++ b/Engine/Plugins/Experimental/Animation/AnimDatabase/Source/AnimDatabase/Public/AnimDatabase.h
@@ -1117,7 +1117,7 @@ public:
 
 	/** Bone to project onto the floor to use as the root location */
 	UPROPERTY(EditAnywhere, Category = "Bones")
-	FName RootLocationBone = TEXT("spine_02");
+	FName RootLocationBone = TEXT("pelvis");
 
 	/** Bone to use for the root facing direction */
 	UPROPERTY(EditAnywhere, Category = "Bones")
@@ -1133,13 +1133,37 @@ public:
 
 public:
 
+	/** If to use a SavGol filter for smoothing. A SavGol filter can work better when applying large amounts of smoothing */
+	UPROPERTY(EditAnywhere, Category = "Settings")
+	bool bUseSavGolFilter = true;
+
+	/** Filter width for the location savgol filter. Larger values create smoother trajectories. */
+	UPROPERTY(EditAnywhere, Category = "Settings", meta = (ForceUnits = "s", ClampMin = "1", UIMin = "1", EditCondition = "bUseSavGolFilter", HideEditConditionToggle))
+	float LocationSavGolFilterWidth = 1.5;
+
+	/** Degree of polynomial to use for the location savgol filter. Larger degrees allow for more bendy trajectories. */
+	UPROPERTY(EditAnywhere, Category = "Settings", meta = (ClampMin = "1", ClampMax = "9", UIMin = "1", UIMax = "9", EditCondition = "bUseSavGolFilter", HideEditConditionToggle))
+	int32 LocationSavGolPolynomialDegree = 3;
+
+	/** Filter width for the direction savgol filter. Larger values create smoother trajectories. */
+	UPROPERTY(EditAnywhere, Category = "Settings", meta = (ForceUnits = "s", ClampMin = "1", UIMin = "1", EditCondition = "bUseSavGolFilter", HideEditConditionToggle))
+	float DirectionSavGolFilterWidth = 1.5f;
+
+	/** Degree of polynomial to use for the direction savgol filter. Larger degrees allow for more bendy trajectories. */
+	UPROPERTY(EditAnywhere, Category = "Settings", meta = (ClampMin = "1", ClampMax = "9", UIMin = "1", UIMax = "9", EditCondition = "bUseSavGolFilter", HideEditConditionToggle))
+	int32 DirectionSavGolPolynomialDegree = 3;
+
+	/** If to use a gaussian window on the SavGol filter weights */
+	UPROPERTY(EditAnywhere, Category = "Settings", meta = (EditCondition = "bUseSavGolFilter", HideEditConditionToggle))
+	bool bUseGaussianWindowedSavGolFilter = true;
+
 	/** Amount of smoothing to apply to the projected bone location */
-	UPROPERTY(EditAnywhere, Category = "Settings", meta = (ForceUnits = "s"))
-	float LocationSmoothingTime = 9.0f;
+	UPROPERTY(EditAnywhere, Category = "Settings", meta = (ForceUnits = "s", EditCondition = "!bUseSavGolFilter", HideEditConditionToggle))
+	float LocationSmoothingTime = 0.15f;
 
 	/** Amount of smoothing to apply to the facing direction */
-	UPROPERTY(EditAnywhere, Category = "Settings", meta = (ForceUnits = "s"))
-	float DirectionSmoothingTime = 15.0f;
+	UPROPERTY(EditAnywhere, Category = "Settings", meta = (ForceUnits = "s", EditCondition = "!bUseSavGolFilter", HideEditConditionToggle))
+	float DirectionSmoothingTime = 0.25f;
 
 public:
 
diff --git a/Engine/Plugins/Experimental/Animation/AnimDatabase/Source/AnimDatabase/Public/AnimDatabaseFrameAttribute.h b/Engine/Plugins/Experimental/Animation/AnimDatabase/Source/AnimDatabase/Public/AnimDatabaseFrameAttribute.h
--- a/Engine/Plugins/Experimental/Animation/AnimDatabase/Source/AnimDatabase/Public/AnimDatabaseFrameAttribute.h
+++ b/Engine/Plugins/Experimental/Animation/AnimDatabase/Source/AnimDatabase/Public/AnimDatabaseFrameAttribute.h
@@ -1450,6 +1450,10 @@ public:
 	UFUNCTION(BlueprintPure, Category = "AnimDatabase")
 	static UE_API FAnimDatabaseFrameAttribute FrameAttributeFilterGaussian(const FAnimDatabaseFrameAttribute& FloatFrameAttribute, const float StandardDeviationInFrames = 3.0f);
 
+	/** Applies a SavGol filter to the values in each range of a frame attribute. Assumes type is compatible */
+	UFUNCTION(BlueprintPure, Category = "AnimDatabase")
+	static UE_API FAnimDatabaseFrameAttribute FrameAttributeFilterSavGol(const FAnimDatabaseFrameAttribute& FloatFrameAttribute, const int32 FilterWidthInFrames = 7, const int32 PolynomialDegree = 3, const bool bGaussianWindowed = true);
+
 	/** Applies a majority vote filter to the values in each range of a frame attribute. Assumes type is compatible */
 	UFUNCTION(BlueprintPure, Category = "AnimDatabase")
 	static UE_API FAnimDatabaseFrameAttribute FrameAttributeFilterMajorityVote(const FAnimDatabaseFrameAttribute& BoolFrameAttribute, const int32 FilterWidthInFrames = 5);
@@ -1527,6 +1531,35 @@ public:
 		const float AngularVelocityMinimumDifference = 10.0f,
 		const float ScalarVelocityMinimumDifference = 0.2f);
 
+	/**
+	 * Computes a root transform by projecting bones onto the ground plane and smoothing the resulting location and direction with a SavGol filter.
+	 *
+	 * @param Database							Database to use
+	 * @param FrameRanges						Ranges in the database to use
+	 * @param RootLocationProjectionBoneName	Name of the bone to project onto the ground plane and use as the root location
+	 * @param RootRotationProjectionBoneName	Name of the bone to project onto the ground plane and use as the root rotation
+	 * @param RootRotationBoneForwardDirection	Local forward direction of the bone to use as the root rotation
+	 * @param RootForwardDirection				Local forward direction of the desired root transform
+	 * @param LocationSavGolFrameNum			Number of frames over which to apply the location savgol filter
+	 * @param LocationSavGolPolynomialDegree	Degree of polynomial to use for the location savgol filter
+	 * @param DirectionSavGolFrameNum			Number of frames over which to apply the location direction filter
+	 * @param DirectionSavGolPolynomialDegree	Degree of polynomial to use for the direction savgol filter
+	 * @param bGaussianWindowed					If to apply a gaussian weighted window to the savgol polynomial fitting
+	 */
+	UFUNCTION(BlueprintPure, Category = "AnimDatabase", meta = (NotBlueprintThreadSafe))
+	static UE_API FAnimDatabaseFrameAttribute FrameAttributeSavGolSmoothedRootTransform(
+		const UAnimDatabase* Database,
+		const FAnimDatabaseFrameRanges& FrameRanges,
+		const FName RootLocationProjectionBoneName = TEXT("spine_02"),
+		const FName RootRotationProjectionBoneName = TEXT("pelvis"),
+		const FVector RootRotationBoneForwardDirection = FVector(0, 1, 0),
+		const FVector RootForwardDirection = FVector(0, 1, 0),
+		const int32 LocationSavGolFrameNum = 61,
+		const int32 LocationSavGolPolynomialDegree = 3,
+		const int32 DirectionSavGolFrameNum = 91,
+		const int32 DirectionSavGolPolynomialDegree = 3,
+		const bool bGaussianWindowed = true);
+		
 	/**
 	 * Computes a root transform by projecting bones onto the ground plane and smoothing the resulting location and direction with a Gaussian filter.
 	 *
@@ -1540,7 +1573,7 @@ public:
 	 * @param RotationSmoothingAmount			Amount of smoothing to apply to the root rotation projection (in frames)
 	 */
 	UFUNCTION(BlueprintPure, Category = "AnimDatabase", meta = (NotBlueprintThreadSafe))
-	static UE_API FAnimDatabaseFrameAttribute FrameAttributeSmoothedRootTransform(
+	static UE_API FAnimDatabaseFrameAttribute FrameAttributeGaussianSmoothedRootTransform(
 		const UAnimDatabase* Database,
 		const FAnimDatabaseFrameRanges& FrameRanges,
 		const FName RootLocationProjectionBoneName = TEXT("spine_02"),
diff --git a/Engine/Plugins/Experimental/Animation/AnimDatabase/Source/AnimDatabaseEditor/Private/AnimDatabaseEditorToolkit.cpp b/Engine/Plugins/Experimental/Animation/AnimDatabase/Source/AnimDatabaseEditor/Private/AnimDatabaseEditorToolkit.cpp
--- a/Engine/Plugins/Experimental/Animation/AnimDatabase/Source/AnimDatabaseEditor/Private/AnimDatabaseEditorToolkit.cpp
+++ b/Engine/Plugins/Experimental/Animation/AnimDatabase/Source/AnimDatabaseEditor/Private/AnimDatabaseEditorToolkit.cpp
@@ -239,6 +239,14 @@ namespace UE::AnimDatabase::Editor
 								})
 								:
 							SNew(SHyperlink).Text(LOCTEXT("NullSequenceName", "null"))
+						]
+						+SHorizontalBox::Slot()
+						.AutoWidth()
+						.HAlign(EHorizontalAlignment::HAlign_Left)
+						.VAlign(EVerticalAlignment::VAlign_Center)
+						.Padding(0.0f, 1.0f, 0.0f, 1.0f)
+						[
+							SNew(STextBlock).Text(Item->bIsMirrored ? LOCTEXT("MirroredAddition", " (mirrored)") : FText::GetEmpty())
 						];
 				}
 				else if (ColumnName == "Start")
@@ -274,7 +282,7 @@ namespace UE::AnimDatabase::Editor
 						.VAlign(EVerticalAlignment::VAlign_Center)
 						.Padding(0.0f, 1.0f, 0.0f, 1.0f)
 						[
-							SNew(SCheckBox).IsChecked(Item->bIsMirrored).IsEnabled(false)
+							SNew(SImage).Image(FCoreStyle::Get().GetBrush("GraphEditor.AlignNodesCenter")).RenderOpacity(Item->bIsMirrored ? 1.0f : 0.0f)
 						];
 				}
 
