AnimDatabase: Added pose function to re-compute velocities.

#Synced-CL 50905131

diff --git a/Engine/Plugins/Experimental/Animation/AnimDatabase/Source/AnimDatabase/Private/AnimDatabasePose.cpp b/Engine/Plugins/Experimental/Animation/AnimDatabase/Source/AnimDatabase/Private/AnimDatabasePose.cpp
--- a/Engine/Plugins/Experimental/Animation/AnimDatabase/Source/AnimDatabase/Private/AnimDatabasePose.cpp
+++ b/Engine/Plugins/Experimental/Animation/AnimDatabase/Source/AnimDatabase/Private/AnimDatabasePose.cpp
@@ -2433,6 +2433,36 @@ namespace UE::AnimDatabase
 			}
 		}
 
+		void RecomputeVelocities(const FPoseRootDataView& InOutRootData, const float DeltaTime)
+		{
+			const int32 FrameNum = InOutRootData.GetFrameNum();
+			if (!ensure(FrameNum >= 2)) { return; }
+
+			InOutRootData.RootLinearVelocities[0] = (FVector3f)(InOutRootData.RootLocations[1] - InOutRootData.RootLocations[0]) / FMath::Max(DeltaTime, UE_SMALL_NUMBER);
+			InOutRootData.RootAngularVelocities[0] = (InOutRootData.RootRotations[1] * InOutRootData.RootRotations[0].Inverse()).GetShortestArcWith(FQuat4f::Identity).ToRotationVector() / FMath::Max(DeltaTime, UE_SMALL_NUMBER);
+			InOutRootData.RootScalarVelocities[0] = Math::VectorLogSafe(Math::VectorDivMax(InOutRootData.RootScales[1], InOutRootData.RootScales[0])) / FMath::Max(DeltaTime, UE_SMALL_NUMBER);
+
+			InOutRootData.RootLinearVelocities[FrameNum - 1] = (FVector3f)(InOutRootData.RootLocations[FrameNum - 1] - InOutRootData.RootLocations[FrameNum - 2]) / FMath::Max(DeltaTime, UE_SMALL_NUMBER);
+			InOutRootData.RootAngularVelocities[FrameNum - 1] = (InOutRootData.RootRotations[FrameNum - 1] * InOutRootData.RootRotations[FrameNum - 2].Inverse()).GetShortestArcWith(FQuat4f::Identity).ToRotationVector() / FMath::Max(DeltaTime, UE_SMALL_NUMBER);
+			InOutRootData.RootScalarVelocities[FrameNum - 1] = Math::VectorLogSafe(Math::VectorDivMax(InOutRootData.RootScales[FrameNum - 1], InOutRootData.RootScales[FrameNum - 2])) / FMath::Max(DeltaTime, UE_SMALL_NUMBER);
+
+			for (int32 FrameIdx = 1; FrameIdx < FrameNum - 1; FrameIdx++)
+			{
+				const FVector3f LinearVelocity0 = (FVector3f)(InOutRootData.RootLocations[FrameIdx] - InOutRootData.RootLocations[FrameIdx - 1]) / FMath::Max(DeltaTime, UE_SMALL_NUMBER);
+				const FVector3f AngularVelocity0 = (InOutRootData.RootRotations[FrameIdx] * InOutRootData.RootRotations[FrameIdx - 1].Inverse()).GetShortestArcWith(FQuat4f::Identity).ToRotationVector() / FMath::Max(DeltaTime, UE_SMALL_NUMBER);
+				const FVector3f ScalarVelocity0 = Math::VectorLogSafe(Math::VectorDivMax(InOutRootData.RootScales[FrameIdx], InOutRootData.RootScales[FrameIdx - 1])) / FMath::Max(DeltaTime, UE_SMALL_NUMBER);
+				
+				const FVector3f LinearVelocity1 = (FVector3f)(InOutRootData.RootLocations[FrameIdx + 1] - InOutRootData.RootLocations[FrameIdx]) / FMath::Max(DeltaTime, UE_SMALL_NUMBER);
+				const FVector3f AngularVelocity1 = (InOutRootData.RootRotations[FrameIdx + 1] * InOutRootData.RootRotations[FrameIdx].Inverse()).GetShortestArcWith(FQuat4f::Identity).ToRotationVector() / FMath::Max(DeltaTime, UE_SMALL_NUMBER);
+				const FVector3f ScalarVelocity1 = Math::VectorLogSafe(Math::VectorDivMax(InOutRootData.RootScales[FrameIdx + 1], InOutRootData.RootScales[FrameIdx])) / FMath::Max(DeltaTime, UE_SMALL_NUMBER);
+
+				InOutRootData.RootLinearVelocities[FrameIdx] = (LinearVelocity0 + LinearVelocity1) / 2.0f;
+				InOutRootData.RootAngularVelocities[FrameIdx] = (AngularVelocity0 + AngularVelocity1) / 2.0f;
+				InOutRootData.RootScalarVelocities[FrameIdx] = (ScalarVelocity0 + ScalarVelocity1) / 2.0f;
+			}
+		}
+
+
 		void ForwardKinematics(
 			const FPoseGlobalBoneDataView& OutPoseGlobalBoneData,
 			const FPoseLocalBoneDataConstView& InPoseLocalBoneData,
diff --git a/Engine/Plugins/Experimental/Animation/AnimDatabase/Source/AnimDatabase/Public/AnimDatabase.h b/Engine/Plugins/Experimental/Animation/AnimDatabase/Source/AnimDatabase/Public/AnimDatabase.h
--- a/Engine/Plugins/Experimental/Animation/AnimDatabase/Source/AnimDatabase/Public/AnimDatabase.h
+++ b/Engine/Plugins/Experimental/Animation/AnimDatabase/Source/AnimDatabase/Public/AnimDatabase.h
@@ -263,6 +263,7 @@ namespace UE::AnimDatabase::Editor
 	struct FQueryEntry
 	{
 		const UAnimSequence* Sequence = nullptr;
+		FString Name;
 		int32 StartFrame = 0;
 		int32 StopFrame = 0;
 		bool bIsMirrored = false;
diff --git a/Engine/Plugins/Experimental/Animation/AnimDatabase/Source/AnimDatabase/Public/AnimDatabasePose.h b/Engine/Plugins/Experimental/Animation/AnimDatabase/Source/AnimDatabase/Public/AnimDatabasePose.h
--- a/Engine/Plugins/Experimental/Animation/AnimDatabase/Source/AnimDatabase/Public/AnimDatabasePose.h
+++ b/Engine/Plugins/Experimental/Animation/AnimDatabase/Source/AnimDatabase/Public/AnimDatabasePose.h
@@ -611,6 +611,9 @@ namespace UE::AnimDatabase
 
 		/** Replaces the given bones in the data in the middle with an interpolation between the first and last frames */
 		UE_API void PatchInterpolate(const FPoseDataView& InOutPoseData, const FFrameRate FrameRate, const bool bApplyToRoot, const UE::Learning::FIndexSet Bones);
+
+		/** Recomputes Velocities using finite difference assuming the first dimension is frames */
+		UE_API void RecomputeVelocities(const FPoseRootDataView& InOutRootData, const float DeltaTime);
 	}
 
 	//--------------------------------------------------
