AnimDatabase: Updated functions for creating frame ranges from AnimSequences

#Synced-CL 50905131

diff --git a/Engine/Plugins/Experimental/Animation/AnimDatabase/Source/AnimDatabase/Private/AnimDatabaseFrameRanges.cpp b/Engine/Plugins/Experimental/Animation/AnimDatabase/Source/AnimDatabase/Private/AnimDatabaseFrameRanges.cpp
--- a/Engine/Plugins/Experimental/Animation/AnimDatabase/Source/AnimDatabase/Private/AnimDatabaseFrameRanges.cpp
+++ b/Engine/Plugins/Experimental/Animation/AnimDatabase/Source/AnimDatabase/Private/AnimDatabaseFrameRanges.cpp
@@ -706,6 +706,16 @@
 	return UAnimDatabaseFramesLibrary::MakeEmptyFrames();
 }
 
+FAnimDatabaseFrameRanges UAnimDatabaseFrameRangesFunction_AnimSequence::MakeFrameRanges_Implementation(const UAnimDatabase* Database, const FAnimDatabaseFrameRanges& FrameRanges) const
+{
+	return UAnimDatabaseFrameRangesLibrary::MakeFrameRangesFromAnimSequence(Database, FrameRanges, AnimSequence);
+}
+
+FAnimDatabaseFrameRanges UAnimDatabaseFrameRangesFunction_AnimSequences::MakeFrameRanges_Implementation(const UAnimDatabase* Database, const FAnimDatabaseFrameRanges& FrameRanges) const
+{
+	return UAnimDatabaseFrameRangesLibrary::MakeFrameRangesFromAnimSequences(Database, FrameRanges, AnimSequences);
+}
+
 FAnimDatabaseFrames UAnimDatabaseFramesFunction_AnimNotify::MakeFrames_Implementation(const UAnimDatabase* Database, const FAnimDatabaseFrameRanges& FrameRanges) const
 {
 	return UAnimDatabaseFramesLibrary::MakeFramesFromAnimNotify(Database, FrameRanges, AnimNotify);
@@ -851,6 +861,71 @@
 	return Out;
 }
 
+FAnimDatabaseFrameRanges UAnimDatabaseFrameRangesLibrary::MakeFrameRangesFromAnimSequence(const UAnimDatabase* Database, const FAnimDatabaseFrameRanges& FrameRanges, UAnimSequence* AnimSequence)
+{
+	if (!Database)
+	{
+		UE_LOGFMT(LogAnimDatabase, Error, "MakeFrameRangesFromSequence: Database is nullptr.");
+		return FAnimDatabaseFrameRanges();
+	}
+
+	if (!FrameRanges.IsValid())
+	{
+		UE_LOGFMT(LogAnimDatabase, Error, "MakeFrameRangesFromSequence: Invalid FrameRanges Object.");
+		return FAnimDatabaseFrameRanges();
+	}
+
+	if (!AnimSequence)
+	{
+		UE_LOGFMT(LogAnimDatabase, Error, "Could not make FrameRanges from Null AnimSequence.");
+		return FAnimDatabaseFrameRanges();
+	}
+
+	const int32 UnmirroredSequenceIdx = Database->FindSequenceIndex(AnimSequence, false);
+	const int32 MirroredSequenceIdx = Database->FindSequenceIndex(AnimSequence, true);
+
+	const bool bContainsUnmirrored = FrameRanges.FrameRangeSet->FindSequenceEntry(UnmirroredSequenceIdx) != INDEX_NONE;
+	const bool bContainsMirrored = FrameRanges.FrameRangeSet->FindSequenceEntry(MirroredSequenceIdx) != INDEX_NONE;
+
+	if (bContainsUnmirrored && bContainsMirrored)
+	{
+		return FrameRangesUnionFromArrayView({
+			MakeFrameRangesFromSequenceIndex(Database, UnmirroredSequenceIdx),
+			MakeFrameRangesFromSequenceIndex(Database, MirroredSequenceIdx)
+			});
+	}
+	else if (bContainsUnmirrored)
+	{
+		return MakeFrameRangesFromSequenceIndex(Database, UnmirroredSequenceIdx);
+	}
+	else if (bContainsMirrored)
+	{
+		return MakeFrameRangesFromSequenceIndex(Database, MirroredSequenceIdx);
+	}
+	else
+	{
+		return MakeEmptyFrameRanges();
+	}
+}
+
+FAnimDatabaseFrameRanges UAnimDatabaseFrameRangesLibrary::MakeFrameRangesFromAnimSequences(const UAnimDatabase* Database, const FAnimDatabaseFrameRanges& FrameRanges, const TArray<UAnimSequence*>& AnimSequences)
+{
+	return MakeFrameRangesFromAnimSequencesArrayView(Database, FrameRanges, AnimSequences);
+}
+
+FAnimDatabaseFrameRanges UAnimDatabaseFrameRangesLibrary::MakeFrameRangesFromAnimSequencesArrayView(const UAnimDatabase* Database, const FAnimDatabaseFrameRanges& FrameRanges, const TArrayView<UAnimSequence* const> AnimSequences)
+{
+	TArray<FAnimDatabaseFrameRanges, TInlineAllocator<64>> FrameRangesSet;
+	FrameRangesSet.Reserve(AnimSequences.Num());
+
+	for (UAnimSequence* AnimSequence : AnimSequences)
+	{
+		FrameRangesSet.Add(MakeFrameRangesFromAnimSequence(Database, FrameRanges, AnimSequence));
+	}
+
+	return FrameRangesUnionFromArrayView(FrameRangesSet);
+}
+
 FAnimDatabaseFrameRanges UAnimDatabaseFrameRangesLibrary::MakeFrameRangesFromSequenceRange(const UAnimDatabase* Database, const int32 SequenceIdx, const int32 RangeStart, const int32 RangeLength)
 {
 	if (!Database)
@@ -885,41 +960,4 @@
 	return Out;
 }
 
-FAnimDatabaseFrameRanges UAnimDatabaseFrameRangesLibrary::MakeFrameRangesFromSequence(const UAnimDatabase* Database, UAnimSequence* AnimSequence)
-{
-	if (!Database)
-	{
-		UE_LOGFMT(LogAnimDatabase, Error, "MakeFrameRangesFromSequence: Database is nullptr.");
-		return FAnimDatabaseFrameRanges();
-	}
-
-	if (!AnimSequence)
-	{
-		UE_LOGFMT(LogAnimDatabase, Error, "Could not make FrameRanges from Null AnimSequence.");
-		return FAnimDatabaseFrameRanges();
-	}
-
-	const int32 UnmirroredSequenceIdx = Database->FindSequenceIndex(AnimSequence, false);
-
-	if (UnmirroredSequenceIdx == INDEX_NONE)
-	{
-		UE_LOGFMT(LogAnimDatabase, Error, "Could not find Anim Sequence '{Name}' in Database.", *AnimSequence->GetName());
-		return FAnimDatabaseFrameRanges();
-	}
-
-	if (Database->GetMirrorDataTable())
-	{
-		const int32 MirroredSequenceIdx = Database->FindSequenceIndex(AnimSequence, true);
-		check(MirroredSequenceIdx != INDEX_NONE);
-		return FrameRangesUnionFromArrayView({
-			MakeFrameRangesFromSequenceIndex(Database, UnmirroredSequenceIdx),
-			MakeFrameRangesFromSequenceIndex(Database, MirroredSequenceIdx)
-			});
-	}
-	else
-	{
-		return MakeFrameRangesFromSequenceIndex(Database, UnmirroredSequenceIdx);
-	}
-}
-
 FAnimDatabaseFrameRanges UAnimDatabaseFrameRangesLibrary::MakeFrameRangesFromAssetName(const UAnimDatabase* Database, const FAnimDatabaseFrameRanges& FrameRanges, const FString& AssetName)
diff --git a/Engine/Plugins/Experimental/Animation/AnimDatabase/Source/AnimDatabase/Public/AnimDatabaseFrameRanges.h b/Engine/Plugins/Experimental/Animation/AnimDatabase/Source/AnimDatabase/Public/AnimDatabaseFrameRanges.h
--- a/Engine/Plugins/Experimental/Animation/AnimDatabase/Source/AnimDatabase/Public/AnimDatabaseFrameRanges.h
+++ b/Engine/Plugins/Experimental/Animation/AnimDatabase/Source/AnimDatabase/Public/AnimDatabaseFrameRanges.h
@@ -338,7 +338,12 @@ public:
 
 	/** Make a set of FrameRanges for a given anim sequence in a database */
 	UFUNCTION(BlueprintPure, Category = "AnimDatabase")
-	static UE_API FAnimDatabaseFrameRanges MakeFrameRangesFromSequence(const UAnimDatabase* Database, UAnimSequence* AnimSequence);
+	static UE_API FAnimDatabaseFrameRanges MakeFrameRangesFromAnimSequence(const UAnimDatabase* Database, const FAnimDatabaseFrameRanges& FrameRanges, UAnimSequence* AnimSequence);
+
+	/** Make a set of FrameRanges for a given anim sequence in a database */
+	UFUNCTION(BlueprintPure, Category = "AnimDatabase")
+	static UE_API FAnimDatabaseFrameRanges MakeFrameRangesFromAnimSequences(const UAnimDatabase* Database, const FAnimDatabaseFrameRanges& FrameRanges, const TArray<UAnimSequence*>& AnimSequences);
+	static UE_API FAnimDatabaseFrameRanges MakeFrameRangesFromAnimSequencesArrayView(const UAnimDatabase* Database, const FAnimDatabaseFrameRanges& FrameRanges, const TArrayView<UAnimSequence* const> AnimSequences);
 
 	/** Make a set of FrameRanges for all sequences with a matching asset name in a set of frame ranges */
 	UFUNCTION(BlueprintPure, Category = "AnimDatabase")
@@ -702,6 +707,38 @@ public:
 	UE_API virtual FAnimDatabaseFrameRanges MakeFrameRanges_Implementation(const UAnimDatabase* Database, const FAnimDatabaseFrameRanges& FrameRanges) const override;
 };
 
+/** Finds all the frame ranges corresponding to the given Anim Sequence */
+UCLASS(BlueprintType, Blueprintable, meta = (DisplayName = "Anim Sequence"))
+class UAnimDatabaseFrameRangesFunction_AnimSequence : public UAnimDatabaseFrameRangesFunction
+{
+	GENERATED_BODY()
+
+public:
+
+	UPROPERTY(EditAnywhere, Category = "Settings")
+	TObjectPtr<UAnimSequence> AnimSequence;
+
+public:
+
+	UE_API virtual FAnimDatabaseFrameRanges MakeFrameRanges_Implementation(const UAnimDatabase* Database, const FAnimDatabaseFrameRanges& FrameRanges) const override;
+};
+
+/** Finds all the frame ranges corresponding to the given Anim Sequences */
+UCLASS(BlueprintType, Blueprintable, meta = (DisplayName = "Anim Sequences"))
+class UAnimDatabaseFrameRangesFunction_AnimSequences : public UAnimDatabaseFrameRangesFunction
+{
+	GENERATED_BODY()
+
+public:
+
+	UPROPERTY(EditAnywhere, Category = "Settings")
+	TArray<TObjectPtr<UAnimSequence>> AnimSequences;
+
+public:
+
+	UE_API virtual FAnimDatabaseFrameRanges MakeFrameRanges_Implementation(const UAnimDatabase* Database, const FAnimDatabaseFrameRanges& FrameRanges) const override;
+};
+
 /** Finds all the frame ranges with a given Anim Notify Class */
 UCLASS(BlueprintType, Blueprintable, meta = (DisplayName = "Anim Notify"))
 class UAnimDatabaseFrameRangesFunction_AnimNotify : public UAnimDatabaseFrameRangesFunction
diff --git a/Engine/Plugins/Experimental/Animation/AnimDatabase/Source/AnimDatabaseEditor/Private/AnimDatabaseEditorMode.cpp b/Engine/Plugins/Experimental/Animation/AnimDatabase/Source/AnimDatabaseEditor/Private/AnimDatabaseEditorMode.cpp
--- a/Engine/Plugins/Experimental/Animation/AnimDatabase/Source/AnimDatabaseEditor/Private/AnimDatabaseEditorMode.cpp
+++ b/Engine/Plugins/Experimental/Animation/AnimDatabase/Source/AnimDatabaseEditor/Private/AnimDatabaseEditorMode.cpp
@@ -235,7 +235,7 @@ namespace UE::AnimDatabase::Editor
 				EditorToolKit->TimelineModel->SetTimelineViewRange(TRange<FFrameNumber>(0, 0));
 				EditorToolKit->TimelineModel->SetTimelineFrameTime(FFrameTime::FromDecimal(0));
 			}
-			if (ActiveRanges.Num() == 1)
+			else if (ActiveRanges.Num() == 1)
 			{
 				const int32 RangeStart = Database->Query->QueryRanges.FrameRangeSet->GetAllRangeStarts()[ActiveRanges[0]];
 				const int32 RangeLength = Database->Query->QueryRanges.FrameRangeSet->GetAllRangeLengths()[ActiveRanges[0]];
