Fix localization warnings in Engine

#Synced-CL 50983539

diff --git a/Engine/Plugins/AudioInsights/Source/AudioInsights/Private/AssetEditorCommands.cpp b/Engine/Plugins/AudioInsights/Source/AudioInsights/Private/AssetEditorCommands.cpp
--- a/Engine/Plugins/AudioInsights/Source/AudioInsights/Private/AssetEditorCommands.cpp
+++ b/Engine/Plugins/AudioInsights/Source/AudioInsights/Private/AssetEditorCommands.cpp
@@ -5,7 +5,7 @@
 #include "Framework/Commands/InputChord.h"
 #include "Framework/Commands/UICommandInfo.h"
 
-#define LOCTEXT_NAMESPACE "AudioInsights"
+#define LOCTEXT_NAMESPACE "AudioInsights.AssetEditorCommands"
 
 namespace UE::Audio::Insights
 {
diff --git a/Engine/Plugins/AudioInsights/Source/AudioInsights/Private/AudioEventLogEditorCommands.cpp b/Engine/Plugins/AudioInsights/Source/AudioInsights/Private/AudioEventLogEditorCommands.cpp
--- a/Engine/Plugins/AudioInsights/Source/AudioInsights/Private/AudioEventLogEditorCommands.cpp
+++ b/Engine/Plugins/AudioInsights/Source/AudioInsights/Private/AudioEventLogEditorCommands.cpp
@@ -5,7 +5,7 @@
 #include "Framework/Commands/InputChord.h"
 #include "Framework/Commands/UICommandInfo.h"
 
-#define LOCTEXT_NAMESPACE "AudioInsights"
+#define LOCTEXT_NAMESPACE "AudioInsights.AudioEventLogEditorCommands"
 
 namespace UE::Audio::Insights
 {
diff --git a/Engine/Plugins/AudioInsights/Source/AudioInsights/Private/SoundDashboardCommands.cpp b/Engine/Plugins/AudioInsights/Source/AudioInsights/Private/SoundDashboardCommands.cpp
--- a/Engine/Plugins/AudioInsights/Source/AudioInsights/Private/SoundDashboardCommands.cpp
+++ b/Engine/Plugins/AudioInsights/Source/AudioInsights/Private/SoundDashboardCommands.cpp
@@ -5,7 +5,7 @@
 #include "Framework/Commands/InputChord.h"
 #include "Framework/Commands/UICommandInfo.h"
 
-#define LOCTEXT_NAMESPACE "AudioInsights"
+#define LOCTEXT_NAMESPACE "AudioInsights.SoundDashboardCommands"
 
 namespace UE::Audio::Insights
 {
diff --git a/Engine/Plugins/Developer/ProjectLauncher/Source/CommonLaunchExtensions/Private/UgsSync/UgsSyncLaunchExtension.cpp b/Engine/Plugins/Developer/ProjectLauncher/Source/CommonLaunchExtensions/Private/UgsSync/UgsSyncLaunchExtension.cpp
--- a/Engine/Plugins/Developer/ProjectLauncher/Source/CommonLaunchExtensions/Private/UgsSync/UgsSyncLaunchExtension.cpp
+++ b/Engine/Plugins/Developer/ProjectLauncher/Source/CommonLaunchExtensions/Private/UgsSync/UgsSyncLaunchExtension.cpp
@@ -16,7 +16,7 @@
 #include "Interfaces/IHttpRequest.h"
 
 
-#define LOCTEXT_NAMESPACE "FBuildSyncLaunchExtensionInstance"
+#define LOCTEXT_NAMESPACE "FUGSSyncLaunchExtensionInstance"
 
 FUGSSyncLaunchExtensionInstance::FUGSSyncLaunchExtensionInstance(FArgs& InArgs)
 	: ProjectLauncher::FCustomUATCommandLaunchExtensionInstance(InArgs)
diff --git a/Engine/Plugins/Experimental/Animation/AnimDatabase/Source/AnimDatabaseEditor/Private/AnimDatabaseEditorToolkit.cpp b/Engine/Plugins/Experimental/Animation/AnimDatabase/Source/AnimDatabaseEditor/Private/AnimDatabaseEditorToolkit.cpp
--- a/Engine/Plugins/Experimental/Animation/AnimDatabase/Source/AnimDatabaseEditor/Private/AnimDatabaseEditorToolkit.cpp
+++ b/Engine/Plugins/Experimental/Animation/AnimDatabase/Source/AnimDatabaseEditor/Private/AnimDatabaseEditorToolkit.cpp
@@ -115,7 +115,7 @@ namespace UE::AnimDatabase::Editor
 				.HeaderRow(
 					SNew(SHeaderRow)
 					+ SHeaderRow::Column("Icon")
-					.DefaultLabel(LOCTEXT("ColumnLabelIcon", ""))
+					.DefaultLabel(FText::GetEmpty())
 					.FillWidth(0.05f)
 					+ SHeaderRow::Column("Sequence")
 					.DefaultLabel(LOCTEXT("ColumnLabelAnimSequence", "Anim Sequence"))
