AnimDatabase: Added some additional frame attribute functions for dealing with angles as well as string functions for debugging.

#Synced-CL 51010646

diff --git a/Engine/Plugins/Experimental/Animation/AnimDatabase/Source/AnimDatabase/Private/AnimDatabase.cpp b/Engine/Plugins/Experimental/Animation/AnimDatabase/Source/AnimDatabase/Private/AnimDatabase.cpp
--- a/Engine/Plugins/Experimental/Animation/AnimDatabase/Source/AnimDatabase/Private/AnimDatabase.cpp
+++ b/Engine/Plugins/Experimental/Animation/AnimDatabase/Source/AnimDatabase/Private/AnimDatabase.cpp
@@ -3730,6 +3730,12 @@ void UAnimDatabase::SampleAttributeData(
 					break;
 				}
 
+				case EAnimDatabaseAttributeType::Angle:
+				{
+					UAnimDatabaseFrameAttributeLibrary::FrameAttributeSampleAngleRadians(OutAttributeData.AttributeData[FrameIdx][FrameAttributeOffset], FrameAttributes[FrameAttributeIdx], SequenceIndices[FrameIdx], SequenceTimes[FrameIdx], FrameRate);
+					break;
+				}
+
 				case EAnimDatabaseAttributeType::Location:
 				{
 					FVector3f LocationValue = FVector3f::ZeroVector;
diff --git a/Engine/Plugins/Experimental/Animation/AnimDatabase/Source/AnimDatabase/Private/AnimDatabaseFrameAttribute.cpp b/Engine/Plugins/Experimental/Animation/AnimDatabase/Source/AnimDatabase/Private/AnimDatabaseFrameAttribute.cpp
--- a/Engine/Plugins/Experimental/Animation/AnimDatabase/Source/AnimDatabase/Private/AnimDatabaseFrameAttribute.cpp
+++ b/Engine/Plugins/Experimental/Animation/AnimDatabase/Source/AnimDatabase/Private/AnimDatabaseFrameAttribute.cpp
@@ -69,6 +69,78 @@ namespace UE::AnimDatabase::FrameAttribute::Private
 			}
 		}
 	}
+
+	static inline FString FrameAttributeFrameToString(const FAnimDatabaseFrameAttribute& FrameAttribute, const int32 FrameIdx)
+	{
+		switch (FrameAttribute.Type)
+		{
+
+		case EAnimDatabaseAttributeType::Null:
+		{
+			return TEXT("null");
+		}
+
+		case EAnimDatabaseAttributeType::Bool:
+		{
+			return (FrameAttribute.FrameAttribute->AttributeData[0][FrameIdx] > 0.5f) ? TEXT("true") : TEXT("false");
+		}
+
+		case EAnimDatabaseAttributeType::Float:
+		case EAnimDatabaseAttributeType::Angle:
+		{
+			return FString::Printf(TEXT("%0.2f"), FrameAttribute.FrameAttribute->AttributeData[0][FrameIdx]);
+		}
+
+		case EAnimDatabaseAttributeType::Location:
+		case EAnimDatabaseAttributeType::Direction:
+		case EAnimDatabaseAttributeType::Scale:
+		case EAnimDatabaseAttributeType::LinearVelocity:
+		case EAnimDatabaseAttributeType::AngularVelocity:
+		case EAnimDatabaseAttributeType::ScalarVelocity:
+		{
+			return FString::Printf(TEXT("[%0.2f %0.2f %0.2f]"),
+				FrameAttribute.FrameAttribute->AttributeData[0][FrameIdx],
+				FrameAttribute.FrameAttribute->AttributeData[1][FrameIdx],
+				FrameAttribute.FrameAttribute->AttributeData[2][FrameIdx]);
+		}
+
+		case EAnimDatabaseAttributeType::Rotation:
+		{
+			return FString::Printf(TEXT("[%0.2f %0.2f %0.2f %0.2f]"),
+				FrameAttribute.FrameAttribute->AttributeData[0][FrameIdx],
+				FrameAttribute.FrameAttribute->AttributeData[1][FrameIdx],
+				FrameAttribute.FrameAttribute->AttributeData[2][FrameIdx],
+				FrameAttribute.FrameAttribute->AttributeData[3][FrameIdx]);
+		}
+
+		case EAnimDatabaseAttributeType::Transform:
+		{
+			return FString::Printf(TEXT("(Location=[%0.2f %0.2f %0.2f] Rotation=[%0.2f %0.2f %0.2f %0.2f] Scale=[%0.2f %0.2f %0.2f])"),
+				FrameAttribute.FrameAttribute->AttributeData[0][FrameIdx],
+				FrameAttribute.FrameAttribute->AttributeData[1][FrameIdx],
+				FrameAttribute.FrameAttribute->AttributeData[2][FrameIdx],
+				FrameAttribute.FrameAttribute->AttributeData[3][FrameIdx],
+				FrameAttribute.FrameAttribute->AttributeData[4][FrameIdx],
+				FrameAttribute.FrameAttribute->AttributeData[5][FrameIdx],
+				FrameAttribute.FrameAttribute->AttributeData[6][FrameIdx],
+				FrameAttribute.FrameAttribute->AttributeData[7][FrameIdx],
+				FrameAttribute.FrameAttribute->AttributeData[8][FrameIdx],
+				FrameAttribute.FrameAttribute->AttributeData[9][FrameIdx]);
+		}
+
+		case EAnimDatabaseAttributeType::Event:
+		{
+			return FrameAttribute.FrameAttribute->AttributeData[0][FrameIdx] > 0.5f ?
+				FString::Printf(TEXT("%0.2f"), FrameAttribute.FrameAttribute->AttributeData[1][FrameIdx]) :
+				TEXT("  --  ");
+		}
+
+		default:
+		{
+			return TEXT("Unknown");
+		}
+		}
+	}
 }
 
 bool FAnimDatabaseFrameAttribute::IsValid() const
@@ -210,6 +282,16 @@ FTransform FAnimDatabaseFrameAttribute::GetAsTransformDouble(const int32 FrameId
 			FrameAttribute->AttributeData[9][FrameIdx]));
 }
 
+float FAnimDatabaseFrameAttribute::GetAsAngleDegrees(const int32 FrameIdx) const
+{
+	return FMath::RadiansToDegrees(FrameAttribute->AttributeData[0][FrameIdx]);
+}
+
+float FAnimDatabaseFrameAttribute::GetAsAngleRadians(const int32 FrameIdx) const
+{
+	return FrameAttribute->AttributeData[0][FrameIdx];
+}
+
 bool FAnimDatabaseFrameAttribute::Serialize(FArchive& Ar)
 {
 	TRACE_CPUPROFILER_EVENT_SCOPE(FAnimDatabaseFrameAttribute::Serialize);
@@ -310,10 +392,37 @@ int32 UAnimDatabaseFrameAttributeLibrary::AttributeTypeSize(const EAnimDatabaseA
 	case EAnimDatabaseAttributeType::Direction: return 3;
 	case EAnimDatabaseAttributeType::Transform: return 10;
 	case EAnimDatabaseAttributeType::Event: return 2;
+	case EAnimDatabaseAttributeType::Angle: return 1;
 	default: checkNoEntry(); return 0;
 	}
 }
 
+FString UAnimDatabaseFrameAttributeLibrary::AttributeTypeName(const EAnimDatabaseAttributeType Type)
+{
+	return AttributeTypeNameInternal(Type);
+}
+
+const TCHAR* UAnimDatabaseFrameAttributeLibrary::AttributeTypeNameInternal(const EAnimDatabaseAttributeType Type)
+{
+	switch (Type)
+	{
+	case EAnimDatabaseAttributeType::Null: return TEXT("Null");
+	case EAnimDatabaseAttributeType::Bool: return TEXT("Bool");
+	case EAnimDatabaseAttributeType::Float: return TEXT("Float");
+	case EAnimDatabaseAttributeType::Location: return TEXT("Location");
+	case EAnimDatabaseAttributeType::Rotation: return TEXT("Rotation");
+	case EAnimDatabaseAttributeType::Scale: return TEXT("Scale");
+	case EAnimDatabaseAttributeType::LinearVelocity: return TEXT("LinearVelocity");
+	case EAnimDatabaseAttributeType::AngularVelocity: return TEXT("AngularVelocity");
+	case EAnimDatabaseAttributeType::ScalarVelocity: return TEXT("ScalarVelocity");
+	case EAnimDatabaseAttributeType::Direction: return TEXT("Direction");
+	case EAnimDatabaseAttributeType::Transform: return TEXT("Transform");
+	case EAnimDatabaseAttributeType::Event: return TEXT("Event");
+	case EAnimDatabaseAttributeType::Angle: return TEXT("Angle");
+	default: return TEXT("Unknown");
+	}
+}
+
 FAnimDatabaseFrameAttribute UAnimDatabaseFrameAttributeLibrary::MakeEmptyFrameAttribute()
 {
 	FAnimDatabaseFrameAttribute Out;
@@ -682,7 +791,7 @@ FAnimDatabaseFrameAttribute UAnimDatabaseFrameAttributeLibrary::MakeEventFrameAt
 				if (bNearestFound)
 				{
 					Out.FrameAttribute->AttributeData[0][FrameOffset + FrameIdx] = 1.0f;
-					Out.FrameAttribute->AttributeData[1][FrameOffset + FrameIdx] = NearestFrameDifference / FrameRate.AsDecimal();
+					Out.FrameAttribute->AttributeData[1][FrameOffset + FrameIdx] = NearestFrameDifference / FMath::Max(FrameRate.AsDecimal(), UE_SMALL_NUMBER);
 				}
 				else
 				{
@@ -794,7 +903,7 @@ FAnimDatabaseFrameAttribute UAnimDatabaseFrameAttributeLibrary::MakeTransformFra
 
 	if (Location.Type != EAnimDatabaseAttributeType::Location || Rotation.Type != EAnimDatabaseAttributeType::Rotation || Scale.Type != EAnimDatabaseAttributeType::Scale)
 	{
-		UE_LOGFMT(LogAnimDatabase, Error, "MakeTransformFrameAttribute: Invalid FrameAttribute Types. Got {Type0}, {Type1}, and {Type2}, Expected Location, Rotation, and Scale.", *UEnum::GetValueAsString(Location.Type), *UEnum::GetValueAsString(Rotation.Type), *UEnum::GetValueAsString(Scale.Type));
+		UE_LOGFMT(LogAnimDatabase, Error, "MakeTransformFrameAttribute: Invalid FrameAttribute Types. Got {Type0}, {Type1}, and {Type2}, Expected Location, Rotation, and Scale.", AttributeTypeNameInternal(Location.Type), AttributeTypeNameInternal(Rotation.Type), AttributeTypeNameInternal(Scale.Type));
 		return FAnimDatabaseFrameAttribute();
 	}
 
@@ -814,7 +923,7 @@ FAnimDatabaseFrameAttribute UAnimDatabaseFrameAttributeLibrary::MakeTransformFra
 
 	if (Location.Type != EAnimDatabaseAttributeType::Location || Rotation.Type != EAnimDatabaseAttributeType::Rotation)
 	{
-		UE_LOGFMT(LogAnimDatabase, Error, "MakeTransformFrameAttributeNoScale: Invalid FrameAttribute Types. Got {Type0} and {Type1}, Expected Location and Rotation.", *UEnum::GetValueAsString(Location.Type), *UEnum::GetValueAsString(Rotation.Type));
+		UE_LOGFMT(LogAnimDatabase, Error, "MakeTransformFrameAttributeNoScale: Invalid FrameAttribute Types. Got {Type0} and {Type1}, Expected Location and Rotation.", AttributeTypeNameInternal(Location.Type), AttributeTypeNameInternal(Rotation.Type));
 		return FAnimDatabaseFrameAttribute();
 	}
 
@@ -966,7 +1075,7 @@ FAnimDatabaseFrameAttribute UAnimDatabaseFrameAttributeLibrary::MakeFloatFrameAt
 
 			for (int32 FrameIdx = 0; FrameIdx < FrameNum; FrameIdx++)
 			{
-				Out.FrameAttribute->AttributeData[0][FrameOffset + FrameIdx] = FrameIdx / FrameRate.AsDecimal();
+				Out.FrameAttribute->AttributeData[0][FrameOffset + FrameIdx] = FrameIdx / FMath::Max(FrameRate.AsDecimal(), UE_SMALL_NUMBER);
 			}
 		});
 
@@ -999,7 +1108,7 @@ FAnimDatabaseFrameAttribute UAnimDatabaseFrameAttributeLibrary::MakeFloatFrameAt
 
 			for (int32 FrameIdx = 0; FrameIdx < FrameNum; FrameIdx++)
 			{
-				Out.FrameAttribute->AttributeData[0][FrameOffset + FrameIdx] = (FrameStart + FrameIdx) / FrameRate.AsDecimal();
+				Out.FrameAttribute->AttributeData[0][FrameOffset + FrameIdx] = (FrameStart + FrameIdx) / FMath::Max(FrameRate.AsDecimal(), UE_SMALL_NUMBER);
 			}
 		});
 
@@ -1043,7 +1152,7 @@ FAnimDatabaseFrameAttribute UAnimDatabaseFrameAttributeLibrary::MakeFloatFrameAt
 				if (Frames.FrameSet->FindNearestInRange(NearestEntry, NearestFrame, FrameDifference, RangeSequence, FrameStart + FrameIdx, FrameStart, FrameNum))
 				{
 					const int32 MatchingFrame = Frames.FrameSet->GetEntryFrame(NearestEntry, NearestFrame);
-					Out.FrameAttribute->AttributeData[0][FrameOffset + FrameIdx] = MatchingFrame / FrameRate.AsDecimal();
+					Out.FrameAttribute->AttributeData[0][FrameOffset + FrameIdx] = MatchingFrame / FMath::Max(FrameRate.AsDecimal(), UE_SMALL_NUMBER);
 				}
 				else
 				{
@@ -1161,6 +1270,77 @@ FAnimDatabaseFrameAttribute UAnimDatabaseFrameAttributeLibrary::MakeScaleFrameAt
 	return Out;
 }
 
+FAnimDatabaseFrameAttribute UAnimDatabaseFrameAttributeLibrary::MakeAngleFrameAttributeFromConstantDegrees(const FAnimDatabaseFrameRanges& FrameRanges, const float Value)
+{
+	if (!FrameRanges.IsValid())
+	{
+		UE_LOGFMT(LogAnimDatabase, Error, "MakeAngleFrameAttributeFromConstantDegrees: Invalid FrameRanges.");
+		return FAnimDatabaseFrameAttribute();
+	}
+
+	FAnimDatabaseFrameAttribute Out = MakeEmptyFrameAttribute();
+	Out.Type = EAnimDatabaseAttributeType::Angle;
+	Out.FrameAttribute->FrameRangeSet = *FrameRanges.FrameRangeSet;
+	Out.FrameAttribute->AttributeData.SetNumUninitialized({ AttributeTypeSize(Out.Type), FrameRanges.FrameRangeSet->GetTotalFrameNum() });
+	UE::Learning::Array::Set(Out.FrameAttribute->AttributeData[0], FMath::DegreesToRadians(Value));
+	return Out;
+}
+
+FAnimDatabaseFrameAttribute UAnimDatabaseFrameAttributeLibrary::MakeAngleFrameAttributeFromConstantRadians(const FAnimDatabaseFrameRanges& FrameRanges, const float Value)
+{
+	if (!FrameRanges.IsValid())
+	{
+		UE_LOGFMT(LogAnimDatabase, Error, "MakeAngleFrameAttributeFromConstantRadians: Invalid FrameRanges.");
+		return FAnimDatabaseFrameAttribute();
+	}
+
+	FAnimDatabaseFrameAttribute Out = MakeEmptyFrameAttribute();
+	Out.Type = EAnimDatabaseAttributeType::Angle;
+	Out.FrameAttribute->FrameRangeSet = *FrameRanges.FrameRangeSet;
+	Out.FrameAttribute->AttributeData.SetNumUninitialized({ AttributeTypeSize(Out.Type), FrameRanges.FrameRangeSet->GetTotalFrameNum() });
+	UE::Learning::Array::Set(Out.FrameAttribute->AttributeData[0], Value);
+	return Out;
+}
+
+FAnimDatabaseFrameAttribute UAnimDatabaseFrameAttributeLibrary::MakeAngleFrameAttributeFromFloatDegrees(const FAnimDatabaseFrameAttribute& A)
+{
+	if (!A.IsValid())
+	{
+		UE_LOGFMT(LogAnimDatabase, Error, "MakeAngleFrameAttributeFromFloatDegrees: Invalid Frame Attribute.");
+		return FAnimDatabaseFrameAttribute();
+	}
+
+	if (A.Type != EAnimDatabaseAttributeType::Float)
+	{
+		UE_LOGFMT(LogAnimDatabase, Error, "MakeAngleFrameAttributeFromFloatDegrees: FrameAttribute Types don't match. Got {Type}. Expected Float.", AttributeTypeNameInternal(A.Type));
+		return FAnimDatabaseFrameAttribute();
+	}
+
+	FAnimDatabaseFrameAttribute Out = FrameAttributeCopy(A);
+	Out.Type = EAnimDatabaseAttributeType::Angle;
+	UE::Learning::FrameAttribute::MulConstantInplace(*Out.FrameAttribute, { UE_PI / 180.f });
+	return Out;
+}
+
+FAnimDatabaseFrameAttribute UAnimDatabaseFrameAttributeLibrary::MakeAngleFrameAttributeFromFloatRadians(const FAnimDatabaseFrameAttribute& A)
+{
+	if (!A.IsValid())
+	{
+		UE_LOGFMT(LogAnimDatabase, Error, "MakeAngleFrameAttributeFromFloatRadians: Invalid Frame Attribute.");
+		return FAnimDatabaseFrameAttribute();
+	}
+
+	if (A.Type != EAnimDatabaseAttributeType::Float)
+	{
+		UE_LOGFMT(LogAnimDatabase, Error, "MakeAngleFrameAttributeFromFloatRadians: FrameAttribute Types don't match. Got {Type}. Expected Float.", AttributeTypeNameInternal(A.Type));
+		return FAnimDatabaseFrameAttribute();
+	}
+
+	FAnimDatabaseFrameAttribute Out = FrameAttributeCopy(A);
+	Out.Type = EAnimDatabaseAttributeType::Angle;
+	return Out;
+}
+
 FAnimDatabaseFrameAttribute UAnimDatabaseFrameAttributeLibrary::MakeRootLocationFrameAttribute(
 	const UAnimDatabase* Database,
 	const FAnimDatabaseFrameRanges& FrameRanges,
@@ -3772,7 +3952,7 @@ bool UAnimDatabaseFrameAttributeLibrary::FrameAttributeBoolAtFrame(bool& OutBool
 
 	if (FrameAttribute.Type != EAnimDatabaseAttributeType::Bool)
 	{
-		UE_LOGFMT(LogAnimDatabase, Error, "FrameAttributeBoolAtFrame: FrameAttribute Types don't match. Got {Type}. Expected Bool.", *UEnum::GetValueAsString(FrameAttribute.Type));
+		UE_LOGFMT(LogAnimDatabase, Error, "FrameAttributeBoolAtFrame: FrameAttribute Types don't match. Got {Type}. Expected Bool.", AttributeTypeNameInternal(FrameAttribute.Type));
 		OutBool = false;
 		return false;
 	}
@@ -3803,7 +3983,7 @@ bool UAnimDatabaseFrameAttributeLibrary::FrameAttributeFloatAtFrame(float& OutVa
 
 	if (FrameAttribute.Type != EAnimDatabaseAttributeType::Float)
 	{
-		UE_LOGFMT(LogAnimDatabase, Error, "FrameAttributeFloatAtFrame: FrameAttribute Types don't match. Got {Type}. Expected Float.", *UEnum::GetValueAsString(FrameAttribute.Type));
+		UE_LOGFMT(LogAnimDatabase, Error, "FrameAttributeFloatAtFrame: FrameAttribute Types don't match. Got {Type}. Expected Float.", AttributeTypeNameInternal(FrameAttribute.Type));
 		OutValue = 0.0f;
 		return false;
 	}
@@ -3847,7 +4027,7 @@ bool UAnimDatabaseFrameAttributeLibrary::FrameAttributeLocationAtFrameFloat(FVec
 
 	if (FrameAttribute.Type != EAnimDatabaseAttributeType::Location)
 	{
-		UE_LOGFMT(LogAnimDatabase, Error, "FrameAttributeLocationAtFrameFloat: FrameAttribute Types don't match. Got {Type}. Expected Location.", *UEnum::GetValueAsString(FrameAttribute.Type));
+		UE_LOGFMT(LogAnimDatabase, Error, "FrameAttributeLocationAtFrameFloat: FrameAttribute Types don't match. Got {Type}. Expected Location.", AttributeTypeNameInternal(FrameAttribute.Type));
 		OutLocation = FVector3f::ZeroVector;
 		return false;
 	}
@@ -3907,7 +4087,7 @@ bool UAnimDatabaseFrameAttributeLibrary::FrameAttributeRotationAtFrameAsQuatFloa
 
 	if (FrameAttribute.Type != EAnimDatabaseAttributeType::Rotation)
 	{
-		UE_LOGFMT(LogAnimDatabase, Error, "FrameAttributeRotationAtFrameAsQuatFloat: FrameAttribute Types don't match. Got {Type}. Expected Rotation.", *UEnum::GetValueAsString(FrameAttribute.Type));
+		UE_LOGFMT(LogAnimDatabase, Error, "FrameAttributeRotationAtFrameAsQuatFloat: FrameAttribute Types don't match. Got {Type}. Expected Rotation.", AttributeTypeNameInternal(FrameAttribute.Type));
 		OutRotation = FQuat4f::Identity;
 		return false;
 	}
@@ -3955,7 +4135,7 @@ bool UAnimDatabaseFrameAttributeLibrary::FrameAttributeScaleAtFrameFloat(FVector
 
 	if (FrameAttribute.Type != EAnimDatabaseAttributeType::Scale)
 	{
-		UE_LOGFMT(LogAnimDatabase, Error, "FrameAttributeScaleAtFrameFloat: FrameAttribute Types don't match. Got {Type}. Expected Scale.", *UEnum::GetValueAsString(FrameAttribute.Type));
+		UE_LOGFMT(LogAnimDatabase, Error, "FrameAttributeScaleAtFrameFloat: FrameAttribute Types don't match. Got {Type}. Expected Scale.", AttributeTypeNameInternal(FrameAttribute.Type));
 		OutScale = FVector3f::OneVector;
 		return false;
 	}
@@ -4002,7 +4182,7 @@ bool UAnimDatabaseFrameAttributeLibrary::FrameAttributeLinearVelocityAtFrameFloa
 
 	if (FrameAttribute.Type != EAnimDatabaseAttributeType::LinearVelocity)
 	{
-		UE_LOGFMT(LogAnimDatabase, Error, "FrameAttributeLinearVelocityAtFrameFloat: FrameAttribute Types don't match. Got {Type}. Expected LinearVelocity.", *UEnum::GetValueAsString(FrameAttribute.Type));
+		UE_LOGFMT(LogAnimDatabase, Error, "FrameAttributeLinearVelocityAtFrameFloat: FrameAttribute Types don't match. Got {Type}. Expected LinearVelocity.", AttributeTypeNameInternal(FrameAttribute.Type));
 		OutLinearVelocity = FVector3f::ZeroVector;
 		return false;
 	}
@@ -4049,7 +4229,7 @@ bool UAnimDatabaseFrameAttributeLibrary::FrameAttributeAngularVelocityAtFrameFlo
 
 	if (FrameAttribute.Type != EAnimDatabaseAttributeType::AngularVelocity)
 	{
-		UE_LOGFMT(LogAnimDatabase, Error, "FrameAttributeAngularVelocityAtFrameFloat: FrameAttribute Types don't match. Got {Type}. Expected AngularVelocity.", *UEnum::GetValueAsString(FrameAttribute.Type));
+		UE_LOGFMT(LogAnimDatabase, Error, "FrameAttributeAngularVelocityAtFrameFloat: FrameAttribute Types don't match. Got {Type}. Expected AngularVelocity.", AttributeTypeNameInternal(FrameAttribute.Type));
 		OutAngularVelocity = FVector3f::ZeroVector;
 		return false;
 	}
@@ -4096,7 +4276,7 @@ bool UAnimDatabaseFrameAttributeLibrary::FrameAttributeScalarVelocityAtFrameFloa
 
 	if (FrameAttribute.Type != EAnimDatabaseAttributeType::ScalarVelocity)
 	{
-		UE_LOGFMT(LogAnimDatabase, Error, "FrameAttributeScalarVelocityAtFrameFloat: FrameAttribute Types don't match. Got {Type}. Expected ScalarVelocity.", *UEnum::GetValueAsString(FrameAttribute.Type));
+		UE_LOGFMT(LogAnimDatabase, Error, "FrameAttributeScalarVelocityAtFrameFloat: FrameAttribute Types don't match. Got {Type}. Expected ScalarVelocity.", AttributeTypeNameInternal(FrameAttribute.Type));
 		OutScalarVelocity = FVector3f::ZeroVector;
 		return false;
 	}
@@ -4143,7 +4323,7 @@ bool UAnimDatabaseFrameAttributeLibrary::FrameAttributeDirectionAtFrameFloat(FVe
 
 	if (FrameAttribute.Type != EAnimDatabaseAttributeType::Direction)
 	{
-		UE_LOGFMT(LogAnimDatabase, Error, "FrameAttributeDirectionAtFrameFloat: FrameAttribute Types don't match. Got {Type}. Expected Direction.", *UEnum::GetValueAsString(FrameAttribute.Type));
+		UE_LOGFMT(LogAnimDatabase, Error, "FrameAttributeDirectionAtFrameFloat: FrameAttribute Types don't match. Got {Type}. Expected Direction.", AttributeTypeNameInternal(FrameAttribute.Type));
 		OutDirection = FVector3f::ForwardVector;
 		return false;
 	}
@@ -4193,7 +4373,7 @@ bool UAnimDatabaseFrameAttributeLibrary::FrameAttributeTransformAtFrameFloat(FTr
 
 	if (FrameAttribute.Type != EAnimDatabaseAttributeType::Transform)
 	{
-		UE_LOGFMT(LogAnimDatabase, Error, "FrameAttributeTransformAtFrameFloat: FrameAttribute Types don't match. Got {Type}. Expected Transform.", *UEnum::GetValueAsString(FrameAttribute.Type));
+		UE_LOGFMT(LogAnimDatabase, Error, "FrameAttributeTransformAtFrameFloat: FrameAttribute Types don't match. Got {Type}. Expected Transform.", AttributeTypeNameInternal(FrameAttribute.Type));
 		OutTransform = FTransform3f::Identity;
 		return false;
 	}
@@ -4238,7 +4418,7 @@ bool UAnimDatabaseFrameAttributeLibrary::FrameAttributeEventAtFrame(bool& OutTim
 
 	if (FrameAttribute.Type != EAnimDatabaseAttributeType::Event)
 	{
-		UE_LOGFMT(LogAnimDatabase, Error, "FrameAttributeEventAtFrame: FrameAttribute Types don't match. Got {Type}. Expected Event.", *UEnum::GetValueAsString(FrameAttribute.Type));
+		UE_LOGFMT(LogAnimDatabase, Error, "FrameAttributeEventAtFrame: FrameAttribute Types don't match. Got {Type}. Expected Event.", AttributeTypeNameInternal(FrameAttribute.Type));
 		OutTimeUntilEventKnown = false;
 		OutTimeUntilEvent = UE_MAX_FLT;
 		return false;
@@ -4261,6 +4441,68 @@ bool UAnimDatabaseFrameAttributeLibrary::FrameAttributeEventAtFrame(bool& OutTim
 	return true;
 }
 
+bool UAnimDatabaseFrameAttributeLibrary::FrameAttributeAngleAtFrameDegrees(float& OutValue, const FAnimDatabaseFrameAttribute& FrameAttribute, const int32 SequenceIdx, const int32 FrameIdx)
+{
+	if (!FrameAttribute.IsValid())
+	{
+		UE_LOGFMT(LogAnimDatabase, Error, "FrameAttributeAngleAtFrameDegrees: Invalid FrameAttribute.");
+		OutValue = 0.0f;
+		return false;
+	}
+
+	if (FrameAttribute.Type != EAnimDatabaseAttributeType::Angle)
+	{
+		UE_LOGFMT(LogAnimDatabase, Error, "FrameAttributeAngleAtFrameDegrees: FrameAttribute Types don't match. Got {Type}. Expected Angle.", AttributeTypeNameInternal(FrameAttribute.Type));
+		OutValue = 0.0f;
+		return false;
+	}
+
+	int32 EntryIdx = INDEX_NONE;
+	int32 RangeIdx = INDEX_NONE;
+	int32 RangeFrame = INDEX_NONE;
+	if (!FrameAttribute.FrameAttribute->FrameRangeSet.Find(EntryIdx, RangeIdx, RangeFrame, SequenceIdx, FrameIdx))
+	{
+		OutValue = 0.0f;
+		return false;
+	}
+
+	const int32 RangeOffset = FrameAttribute.FrameAttribute->FrameRangeSet.GetEntryRangeOffset(EntryIdx, RangeIdx) + RangeFrame;
+
+	OutValue = FMath::RadiansToDegrees(FrameAttribute.FrameAttribute->AttributeData[0][RangeOffset]);
+	return true;
+}
+
+bool UAnimDatabaseFrameAttributeLibrary::FrameAttributeAngleAtFrameRadians(float& OutValue, const FAnimDatabaseFrameAttribute& FrameAttribute, const int32 SequenceIdx, const int32 FrameIdx)
+{
+	if (!FrameAttribute.IsValid())
+	{
+		UE_LOGFMT(LogAnimDatabase, Error, "FrameAttributeAngleAtFrameRadians: Invalid FrameAttribute.");
+		OutValue = 0.0f;
+		return false;
+	}
+
+	if (FrameAttribute.Type != EAnimDatabaseAttributeType::Angle)
+	{
+		UE_LOGFMT(LogAnimDatabase, Error, "FrameAttributeAngleAtFrameRadians: FrameAttribute Types don't match. Got {Type}. Expected Angle.", AttributeTypeNameInternal(FrameAttribute.Type));
+		OutValue = 0.0f;
+		return false;
+	}
+
+	int32 EntryIdx = INDEX_NONE;
+	int32 RangeIdx = INDEX_NONE;
+	int32 RangeFrame = INDEX_NONE;
+	if (!FrameAttribute.FrameAttribute->FrameRangeSet.Find(EntryIdx, RangeIdx, RangeFrame, SequenceIdx, FrameIdx))
+	{
+		OutValue = 0.0f;
+		return false;
+	}
+
+	const int32 RangeOffset = FrameAttribute.FrameAttribute->FrameRangeSet.GetEntryRangeOffset(EntryIdx, RangeIdx) + RangeFrame;
+
+	OutValue = FrameAttribute.FrameAttribute->AttributeData[0][RangeOffset];
+	return true;
+}
+
 bool UAnimDatabaseFrameAttributeLibrary::FrameAttributeLocationsAtFrame(TArray<FVector>& OutLocations, const TArray<FAnimDatabaseFrameAttribute>& FrameAttributes, const int32 SequenceIdx, const int32 FrameIdx)
 {
 	OutLocations.SetNum(FrameAttributes.Num());
@@ -4335,7 +4577,7 @@ bool UAnimDatabaseFrameAttributeLibrary::FrameAttributeSampleBool(bool& OutBool,
 
 	if (FrameAttribute.Type != EAnimDatabaseAttributeType::Bool)
 	{
-		UE_LOGFMT(LogAnimDatabase, Error, "FrameAttributeSampleBool: FrameAttribute Types don't match. Got {Type}. Expected Bool.", *UEnum::GetValueAsString(FrameAttribute.Type));
+		UE_LOGFMT(LogAnimDatabase, Error, "FrameAttributeSampleBool: FrameAttribute Types don't match. Got {Type}. Expected Bool.", AttributeTypeNameInternal(FrameAttribute.Type));
 		OutBool = false;
 		return false;
 	}
@@ -4343,7 +4585,7 @@ bool UAnimDatabaseFrameAttributeLibrary::FrameAttributeSampleBool(bool& OutBool,
 	int32 EntryIdx = INDEX_NONE;
 	int32 RangeIdx = INDEX_NONE;
 	float RangeTime = 1.0f;
-	if (!FrameAttribute.FrameAttribute->FrameRangeSet.FindTime(EntryIdx, RangeIdx, RangeTime, SequenceIdx, SequenceTime, 1.0f / FrameRate.AsDecimal()))
+	if (!FrameAttribute.FrameAttribute->FrameRangeSet.FindTime(EntryIdx, RangeIdx, RangeTime, SequenceIdx, SequenceTime, 1.0f / FMath::Max(FrameRate.AsDecimal(), UE_SMALL_NUMBER)))
 	{
 		OutBool = false;
 		return false;
@@ -4351,7 +4593,7 @@ bool UAnimDatabaseFrameAttributeLibrary::FrameAttributeSampleBool(bool& OutBool,
 
 	const int32 RangeOffset = FrameAttribute.FrameAttribute->FrameRangeSet.GetEntryRangeOffset(EntryIdx, RangeIdx);
 	const int32 RangeLength = FrameAttribute.FrameAttribute->FrameRangeSet.GetEntryRangeLength(EntryIdx, RangeIdx);
-	const float RangeFrameTime = 1.0f / FrameRate.AsDecimal();
+	const float RangeFrameTime = 1.0f / FMath::Max(FrameRate.AsDecimal(), UE_SMALL_NUMBER);
 	const int64 SampleFrameNearest = UE::AnimDatabase::Math::ComputeNearestSampleFrame(RangeTime / RangeFrameTime, RangeLength);
 
 	OutBool = FrameAttribute.FrameAttribute->AttributeData[0][RangeOffset + SampleFrameNearest] == 1.0f;
@@ -4369,7 +4611,7 @@ bool UAnimDatabaseFrameAttributeLibrary::FrameAttributeSampleFloat(float& OutVal
 
 	if (FrameAttribute.Type != EAnimDatabaseAttributeType::Float)
 	{
-		UE_LOGFMT(LogAnimDatabase, Error, "FrameAttributeSampleFloat: FrameAttribute Types don't match. Got {Type}. Expected Float.", *UEnum::GetValueAsString(FrameAttribute.Type));
+		UE_LOGFMT(LogAnimDatabase, Error, "FrameAttributeSampleFloat: FrameAttribute Types don't match. Got {Type}. Expected Float.", AttributeTypeNameInternal(FrameAttribute.Type));
 		OutValue = 0.0f;
 		return false;
 	}
@@ -4377,7 +4619,7 @@ bool UAnimDatabaseFrameAttributeLibrary::FrameAttributeSampleFloat(float& OutVal
 	int32 EntryIdx = INDEX_NONE;
 	int32 RangeIdx = INDEX_NONE;
 	float RangeTime = 1.0f;
-	if (!FrameAttribute.FrameAttribute->FrameRangeSet.FindTime(EntryIdx, RangeIdx, RangeTime, SequenceIdx, SequenceTime, 1.0f / FrameRate.AsDecimal()))
+	if (!FrameAttribute.FrameAttribute->FrameRangeSet.FindTime(EntryIdx, RangeIdx, RangeTime, SequenceIdx, SequenceTime, 1.0f / FMath::Max(FrameRate.AsDecimal(), UE_SMALL_NUMBER)))
 	{
 		OutValue = 0.0f;
 		return false;
@@ -4385,7 +4627,7 @@ bool UAnimDatabaseFrameAttributeLibrary::FrameAttributeSampleFloat(float& OutVal
 
 	const int32 RangeOffset = FrameAttribute.FrameAttribute->FrameRangeSet.GetEntryRangeOffset(EntryIdx, RangeIdx);
 	const int32 RangeLength = FrameAttribute.FrameAttribute->FrameRangeSet.GetEntryRangeLength(EntryIdx, RangeIdx);
-	const float RangeFrameTime = 1.0f / FrameRate.AsDecimal();
+	const float RangeFrameTime = 1.0f / FMath::Max(FrameRate.AsDecimal(), UE_SMALL_NUMBER);
 
 	int64 SampleFrame0, SampleFrame1, SampleFrame2, SampleFrame3;
 	float SampleAlpha;
@@ -4486,7 +4728,7 @@ bool UAnimDatabaseFrameAttributeLibrary::FrameAttributeSampleLocationFloat(FVect
 
 	if (FrameAttribute.Type != EAnimDatabaseAttributeType::Location)
 	{
-		UE_LOGFMT(LogAnimDatabase, Error, "FrameAttributeSampleLocationFloat: FrameAttribute Types don't match. Got {Type}. Expected Location.", *UEnum::GetValueAsString(FrameAttribute.Type));
+		UE_LOGFMT(LogAnimDatabase, Error, "FrameAttributeSampleLocationFloat: FrameAttribute Types don't match. Got {Type}. Expected Location.", AttributeTypeNameInternal(FrameAttribute.Type));
 		OutLocation = FVector3f::ZeroVector;
 		return false;
 	}
@@ -4494,7 +4736,7 @@ bool UAnimDatabaseFrameAttributeLibrary::FrameAttributeSampleLocationFloat(FVect
 	int32 EntryIdx = INDEX_NONE;
 	int32 RangeIdx = INDEX_NONE;
 	float RangeTime = 1.0f;
-	if (!FrameAttribute.FrameAttribute->FrameRangeSet.FindTime(EntryIdx, RangeIdx, RangeTime, SequenceIdx, SequenceTime, 1.0f / FrameRate.AsDecimal()))
+	if (!FrameAttribute.FrameAttribute->FrameRangeSet.FindTime(EntryIdx, RangeIdx, RangeTime, SequenceIdx, SequenceTime, 1.0f / FMath::Max(FrameRate.AsDecimal(), UE_SMALL_NUMBER)))
 	{
 		OutLocation = FVector3f::ZeroVector;
 		return false;
@@ -4502,7 +4744,7 @@ bool UAnimDatabaseFrameAttributeLibrary::FrameAttributeSampleLocationFloat(FVect
 
 	const int32 RangeOffset = FrameAttribute.FrameAttribute->FrameRangeSet.GetEntryRangeOffset(EntryIdx, RangeIdx);
 	const int32 RangeLength = FrameAttribute.FrameAttribute->FrameRangeSet.GetEntryRangeLength(EntryIdx, RangeIdx);
-	const float RangeFrameTime = 1.0f / FrameRate.AsDecimal();
+	const float RangeFrameTime = 1.0f / FMath::Max(FrameRate.AsDecimal(), UE_SMALL_NUMBER);
 
 	int64 SampleFrame0, SampleFrame1, SampleFrame2, SampleFrame3;
 	float SampleAlpha;
@@ -4596,7 +4838,7 @@ bool UAnimDatabaseFrameAttributeLibrary::FrameAttributeSampleRotationAsQuatFloat
 
 	if (FrameAttribute.Type != EAnimDatabaseAttributeType::Rotation)
 	{
-		UE_LOGFMT(LogAnimDatabase, Error, "FrameAttributeSampleRotationAsQuatFloat: FrameAttribute Types don't match. Got {Type}. Expected Rotation.", *UEnum::GetValueAsString(FrameAttribute.Type));
+		UE_LOGFMT(LogAnimDatabase, Error, "FrameAttributeSampleRotationAsQuatFloat: FrameAttribute Types don't match. Got {Type}. Expected Rotation.", AttributeTypeNameInternal(FrameAttribute.Type));
 		OutRotation = FQuat4f::Identity;
 		return false;
 	}
@@ -4604,7 +4846,7 @@ bool UAnimDatabaseFrameAttributeLibrary::FrameAttributeSampleRotationAsQuatFloat
 	int32 EntryIdx = INDEX_NONE;
 	int32 RangeIdx = INDEX_NONE;
 	float RangeTime = 1.0f;
-	if (!FrameAttribute.FrameAttribute->FrameRangeSet.FindTime(EntryIdx, RangeIdx, RangeTime, SequenceIdx, SequenceTime, 1.0f / FrameRate.AsDecimal()))
+	if (!FrameAttribute.FrameAttribute->FrameRangeSet.FindTime(EntryIdx, RangeIdx, RangeTime, SequenceIdx, SequenceTime, 1.0f / FMath::Max(FrameRate.AsDecimal(), UE_SMALL_NUMBER)))
 	{
 		OutRotation = FQuat4f::Identity;
 		return false;
@@ -4612,7 +4854,7 @@ bool UAnimDatabaseFrameAttributeLibrary::FrameAttributeSampleRotationAsQuatFloat
 
 	const int32 RangeOffset = FrameAttribute.FrameAttribute->FrameRangeSet.GetEntryRangeOffset(EntryIdx, RangeIdx);
 	const int32 RangeLength = FrameAttribute.FrameAttribute->FrameRangeSet.GetEntryRangeLength(EntryIdx, RangeIdx);
-	const float RangeFrameTime = 1.0f / FrameRate.AsDecimal();
+	const float RangeFrameTime = 1.0f / FMath::Max(FrameRate.AsDecimal(), UE_SMALL_NUMBER);
 
 	int64 SampleFrame0, SampleFrame1, SampleFrame2, SampleFrame3;
 	float SampleAlpha;
@@ -4693,7 +4935,7 @@ bool UAnimDatabaseFrameAttributeLibrary::FrameAttributeSampleScaleFloat(FVector3
 
 	if (FrameAttribute.Type != EAnimDatabaseAttributeType::Scale)
 	{
-		UE_LOGFMT(LogAnimDatabase, Error, "FrameAttributeSampleScaleFloat: FrameAttribute Types don't match. Got {Type}. Expected Scale.", *UEnum::GetValueAsString(FrameAttribute.Type));
+		UE_LOGFMT(LogAnimDatabase, Error, "FrameAttributeSampleScaleFloat: FrameAttribute Types don't match. Got {Type}. Expected Scale.", AttributeTypeNameInternal(FrameAttribute.Type));
 		OutScale = FVector3f::OneVector;
 		return false;
 	}
@@ -4701,7 +4943,7 @@ bool UAnimDatabaseFrameAttributeLibrary::FrameAttributeSampleScaleFloat(FVector3
 	int32 EntryIdx = INDEX_NONE;
 	int32 RangeIdx = INDEX_NONE;
 	float RangeTime = 1.0f;
-	if (!FrameAttribute.FrameAttribute->FrameRangeSet.FindTime(EntryIdx, RangeIdx, RangeTime, SequenceIdx, SequenceTime, 1.0f / FrameRate.AsDecimal()))
+	if (!FrameAttribute.FrameAttribute->FrameRangeSet.FindTime(EntryIdx, RangeIdx, RangeTime, SequenceIdx, SequenceTime, 1.0f / FMath::Max(FrameRate.AsDecimal(), UE_SMALL_NUMBER)))
 	{
 		OutScale = FVector3f::OneVector;
 		return false;
@@ -4709,7 +4951,7 @@ bool UAnimDatabaseFrameAttributeLibrary::FrameAttributeSampleScaleFloat(FVector3
 
 	const int32 RangeOffset = FrameAttribute.FrameAttribute->FrameRangeSet.GetEntryRangeOffset(EntryIdx, RangeIdx);
 	const int32 RangeLength = FrameAttribute.FrameAttribute->FrameRangeSet.GetEntryRangeLength(EntryIdx, RangeIdx);
-	const float RangeFrameTime = 1.0f / FrameRate.AsDecimal();
+	const float RangeFrameTime = 1.0f / FMath::Max(FrameRate.AsDecimal(), UE_SMALL_NUMBER);
 
 	int64 SampleFrame0, SampleFrame1, SampleFrame2, SampleFrame3;
 	float SampleAlpha;
@@ -4790,7 +5032,7 @@ bool UAnimDatabaseFrameAttributeLibrary::FrameAttributeSampleLinearVelocityFloat
 
 	if (FrameAttribute.Type != EAnimDatabaseAttributeType::LinearVelocity)
 	{
-		UE_LOGFMT(LogAnimDatabase, Error, "FrameAttributeSampleLinearVelocityFloat: FrameAttribute Types don't match. Got {Type}. Expected LinearVelocity.", *UEnum::GetValueAsString(FrameAttribute.Type));
+		UE_LOGFMT(LogAnimDatabase, Error, "FrameAttributeSampleLinearVelocityFloat: FrameAttribute Types don't match. Got {Type}. Expected LinearVelocity.", AttributeTypeNameInternal(FrameAttribute.Type));
 		OutLinearVelocity = FVector3f::ZeroVector;
 		return false;
 	}
@@ -4798,7 +5040,7 @@ bool UAnimDatabaseFrameAttributeLibrary::FrameAttributeSampleLinearVelocityFloat
 	int32 EntryIdx = INDEX_NONE;
 	int32 RangeIdx = INDEX_NONE;
 	float RangeTime = 1.0f;
-	if (!FrameAttribute.FrameAttribute->FrameRangeSet.FindTime(EntryIdx, RangeIdx, RangeTime, SequenceIdx, SequenceTime, 1.0f / FrameRate.AsDecimal()))
+	if (!FrameAttribute.FrameAttribute->FrameRangeSet.FindTime(EntryIdx, RangeIdx, RangeTime, SequenceIdx, SequenceTime, 1.0f / FMath::Max(FrameRate.AsDecimal(), UE_SMALL_NUMBER)))
 	{
 		OutLinearVelocity = FVector3f::ZeroVector;
 		return false;
@@ -4806,7 +5048,7 @@ bool UAnimDatabaseFrameAttributeLibrary::FrameAttributeSampleLinearVelocityFloat
 
 	const int32 RangeOffset = FrameAttribute.FrameAttribute->FrameRangeSet.GetEntryRangeOffset(EntryIdx, RangeIdx);
 	const int32 RangeLength = FrameAttribute.FrameAttribute->FrameRangeSet.GetEntryRangeLength(EntryIdx, RangeIdx);
-	const float RangeFrameTime = 1.0f / FrameRate.AsDecimal();
+	const float RangeFrameTime = 1.0f / FMath::Max(FrameRate.AsDecimal(), UE_SMALL_NUMBER);
 
 	int64 SampleFrame0, SampleFrame1, SampleFrame2, SampleFrame3;
 	float SampleAlpha;
@@ -4887,7 +5129,7 @@ bool UAnimDatabaseFrameAttributeLibrary::FrameAttributeSampleAngularVelocityFloa
 
 	if (FrameAttribute.Type != EAnimDatabaseAttributeType::AngularVelocity)
 	{
-		UE_LOGFMT(LogAnimDatabase, Error, "FrameAttributeSampleAngularVelocityFloat: FrameAttribute Types don't match. Got {Type}. Expected AngularVelocity.", *UEnum::GetValueAsString(FrameAttribute.Type));
+		UE_LOGFMT(LogAnimDatabase, Error, "FrameAttributeSampleAngularVelocityFloat: FrameAttribute Types don't match. Got {Type}. Expected AngularVelocity.", AttributeTypeNameInternal(FrameAttribute.Type));
 		OutAngularVelocity = FVector3f::ZeroVector;
 		return false;
 	}
@@ -4895,7 +5137,7 @@ bool UAnimDatabaseFrameAttributeLibrary::FrameAttributeSampleAngularVelocityFloa
 	int32 EntryIdx = INDEX_NONE;
 	int32 RangeIdx = INDEX_NONE;
 	float RangeTime = 1.0f;
-	if (!FrameAttribute.FrameAttribute->FrameRangeSet.FindTime(EntryIdx, RangeIdx, RangeTime, SequenceIdx, SequenceTime, 1.0f / FrameRate.AsDecimal()))
+	if (!FrameAttribute.FrameAttribute->FrameRangeSet.FindTime(EntryIdx, RangeIdx, RangeTime, SequenceIdx, SequenceTime, 1.0f / FMath::Max(FrameRate.AsDecimal(), UE_SMALL_NUMBER)))
 	{
 		OutAngularVelocity = FVector3f::ZeroVector;
 		return false;
@@ -4903,7 +5145,7 @@ bool UAnimDatabaseFrameAttributeLibrary::FrameAttributeSampleAngularVelocityFloa
 
 	const int32 RangeOffset = FrameAttribute.FrameAttribute->FrameRangeSet.GetEntryRangeOffset(EntryIdx, RangeIdx);
 	const int32 RangeLength = FrameAttribute.FrameAttribute->FrameRangeSet.GetEntryRangeLength(EntryIdx, RangeIdx);
-	const float RangeFrameTime = 1.0f / FrameRate.AsDecimal();
+	const float RangeFrameTime = 1.0f / FMath::Max(FrameRate.AsDecimal(), UE_SMALL_NUMBER);
 
 	int64 SampleFrame0, SampleFrame1, SampleFrame2, SampleFrame3;
 	float SampleAlpha;
@@ -4984,7 +5226,7 @@ bool UAnimDatabaseFrameAttributeLibrary::FrameAttributeSampleScalarVelocityFloat
 
 	if (FrameAttribute.Type != EAnimDatabaseAttributeType::ScalarVelocity)
 	{
-		UE_LOGFMT(LogAnimDatabase, Error, "FrameAttributeSampleScalarVelocityFloat: FrameAttribute Types don't match. Got {Type}. Expected ScalarVelocity.", *UEnum::GetValueAsString(FrameAttribute.Type));
+		UE_LOGFMT(LogAnimDatabase, Error, "FrameAttributeSampleScalarVelocityFloat: FrameAttribute Types don't match. Got {Type}. Expected ScalarVelocity.", AttributeTypeNameInternal(FrameAttribute.Type));
 		OutScalarVelocity = FVector3f::ZeroVector;
 		return false;
 	}
@@ -4992,7 +5234,7 @@ bool UAnimDatabaseFrameAttributeLibrary::FrameAttributeSampleScalarVelocityFloat
 	int32 EntryIdx = INDEX_NONE;
 	int32 RangeIdx = INDEX_NONE;
 	float RangeTime = 1.0f;
-	if (!FrameAttribute.FrameAttribute->FrameRangeSet.FindTime(EntryIdx, RangeIdx, RangeTime, SequenceIdx, SequenceTime, 1.0f / FrameRate.AsDecimal()))
+	if (!FrameAttribute.FrameAttribute->FrameRangeSet.FindTime(EntryIdx, RangeIdx, RangeTime, SequenceIdx, SequenceTime, 1.0f / FMath::Max(FrameRate.AsDecimal(), UE_SMALL_NUMBER)))
 	{
 		OutScalarVelocity = FVector3f::ZeroVector;
 		return false;
@@ -5000,7 +5242,7 @@ bool UAnimDatabaseFrameAttributeLibrary::FrameAttributeSampleScalarVelocityFloat
 
 	const int32 RangeOffset = FrameAttribute.FrameAttribute->FrameRangeSet.GetEntryRangeOffset(EntryIdx, RangeIdx);
 	const int32 RangeLength = FrameAttribute.FrameAttribute->FrameRangeSet.GetEntryRangeLength(EntryIdx, RangeIdx);
-	const float RangeFrameTime = 1.0f / FrameRate.AsDecimal();
+	const float RangeFrameTime = 1.0f / FMath::Max(FrameRate.AsDecimal(), UE_SMALL_NUMBER);
 
 	int64 SampleFrame0, SampleFrame1, SampleFrame2, SampleFrame3;
 	float SampleAlpha;
@@ -5081,7 +5323,7 @@ bool UAnimDatabaseFrameAttributeLibrary::FrameAttributeSampleDirectionFloat(FVec
 
 	if (FrameAttribute.Type != EAnimDatabaseAttributeType::Direction)
 	{
-		UE_LOGFMT(LogAnimDatabase, Error, "FrameAttributeSampleDirectionFloat: FrameAttribute Types don't match. Got {Type}. Expected Direction.", *UEnum::GetValueAsString(FrameAttribute.Type));
+		UE_LOGFMT(LogAnimDatabase, Error, "FrameAttributeSampleDirectionFloat: FrameAttribute Types don't match. Got {Type}. Expected Direction.", AttributeTypeNameInternal(FrameAttribute.Type));
 		OutDirection = FVector3f::ForwardVector;
 		return false;
 	}
@@ -5089,7 +5331,7 @@ bool UAnimDatabaseFrameAttributeLibrary::FrameAttributeSampleDirectionFloat(FVec
 	int32 EntryIdx = INDEX_NONE;
 	int32 RangeIdx = INDEX_NONE;
 	float RangeTime = 1.0f;
-	if (!FrameAttribute.FrameAttribute->FrameRangeSet.FindTime(EntryIdx, RangeIdx, RangeTime, SequenceIdx, SequenceTime, 1.0f / FrameRate.AsDecimal()))
+	if (!FrameAttribute.FrameAttribute->FrameRangeSet.FindTime(EntryIdx, RangeIdx, RangeTime, SequenceIdx, SequenceTime, 1.0f / FMath::Max(FrameRate.AsDecimal(), UE_SMALL_NUMBER)))
 	{
 		OutDirection = FVector3f::ForwardVector;
 		return false;
@@ -5097,7 +5339,7 @@ bool UAnimDatabaseFrameAttributeLibrary::FrameAttributeSampleDirectionFloat(FVec
 
 	const int32 RangeOffset = FrameAttribute.FrameAttribute->FrameRangeSet.GetEntryRangeOffset(EntryIdx, RangeIdx);
 	const int32 RangeLength = FrameAttribute.FrameAttribute->FrameRangeSet.GetEntryRangeLength(EntryIdx, RangeIdx);
-	const float RangeFrameTime = 1.0f / FrameRate.AsDecimal();
+	const float RangeFrameTime = 1.0f / FMath::Max(FrameRate.AsDecimal(), UE_SMALL_NUMBER);
 
 	int64 SampleFrame0, SampleFrame1, SampleFrame2, SampleFrame3;
 	float SampleAlpha;
@@ -5181,7 +5423,7 @@ bool UAnimDatabaseFrameAttributeLibrary::FrameAttributeSampleTransformFloat(FTra
 
 	if (FrameAttribute.Type != EAnimDatabaseAttributeType::Transform)
 	{
-		UE_LOGFMT(LogAnimDatabase, Error, "FrameAttributeSampleTransformFloat: FrameAttribute Types don't match. Got {Type}. Expected Transform.", *UEnum::GetValueAsString(FrameAttribute.Type));
+		UE_LOGFMT(LogAnimDatabase, Error, "FrameAttributeSampleTransformFloat: FrameAttribute Types don't match. Got {Type}. Expected Transform.", AttributeTypeNameInternal(FrameAttribute.Type));
 		OutTransform = FTransform3f::Identity;
 		return false;
 	}
@@ -5189,7 +5431,7 @@ bool UAnimDatabaseFrameAttributeLibrary::FrameAttributeSampleTransformFloat(FTra
 	int32 EntryIdx = INDEX_NONE;
 	int32 RangeIdx = INDEX_NONE;
 	float RangeTime = 1.0f;
-	if (!FrameAttribute.FrameAttribute->FrameRangeSet.FindTime(EntryIdx, RangeIdx, RangeTime, SequenceIdx, SequenceTime, 1.0f / FrameRate.AsDecimal()))
+	if (!FrameAttribute.FrameAttribute->FrameRangeSet.FindTime(EntryIdx, RangeIdx, RangeTime, SequenceIdx, SequenceTime, 1.0f / FMath::Max(FrameRate.AsDecimal(), UE_SMALL_NUMBER)))
 	{
 		OutTransform = FTransform3f::Identity;
 		return false;
@@ -5197,7 +5439,7 @@ bool UAnimDatabaseFrameAttributeLibrary::FrameAttributeSampleTransformFloat(FTra
 
 	const int32 RangeOffset = FrameAttribute.FrameAttribute->FrameRangeSet.GetEntryRangeOffset(EntryIdx, RangeIdx);
 	const int32 RangeLength = FrameAttribute.FrameAttribute->FrameRangeSet.GetEntryRangeLength(EntryIdx, RangeIdx);
-	const float RangeFrameTime = 1.0f / FrameRate.AsDecimal();
+	const float RangeFrameTime = 1.0f / FMath::Max(FrameRate.AsDecimal(), UE_SMALL_NUMBER);
 
 	int64 SampleFrame0, SampleFrame1, SampleFrame2, SampleFrame3;
 	float SampleAlpha;
@@ -5326,7 +5568,7 @@ bool UAnimDatabaseFrameAttributeLibrary::FrameAttributeSampleEvent(bool& OutTime
 
 	if (FrameAttribute.Type != EAnimDatabaseAttributeType::Event)
 	{
-		UE_LOGFMT(LogAnimDatabase, Error, "FrameAttributeSampleEvent: FrameAttribute Types don't match. Got {Type}. Expected Event.", *UEnum::GetValueAsString(FrameAttribute.Type));
+		UE_LOGFMT(LogAnimDatabase, Error, "FrameAttributeSampleEvent: FrameAttribute Types don't match. Got {Type}. Expected Event.", AttributeTypeNameInternal(FrameAttribute.Type));
 		OutTimeUntilEventKnown = false;
 		OutTimeUntilEvent = UE_MAX_FLT;
 		return false;
@@ -5335,7 +5577,7 @@ bool UAnimDatabaseFrameAttributeLibrary::FrameAttributeSampleEvent(bool& OutTime
 	int32 EntryIdx = INDEX_NONE;
 	int32 RangeIdx = INDEX_NONE;
 	float RangeTime = 1.0f;
-	if (!FrameAttribute.FrameAttribute->FrameRangeSet.FindTime(EntryIdx, RangeIdx, RangeTime, SequenceIdx, SequenceTime, 1.0f / FrameRate.AsDecimal()))
+	if (!FrameAttribute.FrameAttribute->FrameRangeSet.FindTime(EntryIdx, RangeIdx, RangeTime, SequenceIdx, SequenceTime, 1.0f / FMath::Max(FrameRate.AsDecimal(), UE_SMALL_NUMBER)))
 	{
 		OutTimeUntilEventKnown = false;
 		OutTimeUntilEvent = UE_MAX_FLT;
@@ -5344,7 +5586,7 @@ bool UAnimDatabaseFrameAttributeLibrary::FrameAttributeSampleEvent(bool& OutTime
 
 	const int32 RangeOffset = FrameAttribute.FrameAttribute->FrameRangeSet.GetEntryRangeOffset(EntryIdx, RangeIdx);
 	const int32 RangeLength = FrameAttribute.FrameAttribute->FrameRangeSet.GetEntryRangeLength(EntryIdx, RangeIdx);
-	const float RangeFrameTime = 1.0f / FrameRate.AsDecimal();
+	const float RangeFrameTime = 1.0f / FMath::Max(FrameRate.AsDecimal(), UE_SMALL_NUMBER);
 
 	int64 SampleFrame1, SampleFrame2;
 	float SampleAlpha;
@@ -5377,74 +5619,242 @@ bool UAnimDatabaseFrameAttributeLibrary::FrameAttributeSampleEvent(bool& OutTime
 	return true;
 }
 
-bool UAnimDatabaseFrameAttributeLibrary::FrameAttributeSampleLocations(TArray<FVector>& OutLocations, const TArray<FAnimDatabaseFrameAttribute>& FrameAttributes, const int32 SequenceIdx, const float SequenceTime, const FFrameRate& FrameRate)
-{
-	OutLocations.SetNum(FrameAttributes.Num());
-	return FrameAttributeSampleLocationsToArrayView(OutLocations, FrameAttributes, SequenceIdx, SequenceTime, FrameRate);
-}
-
-bool UAnimDatabaseFrameAttributeLibrary::FrameAttributeSampleLocationsToArrayView(const TArrayView<FVector> OutLocations, const TArrayView<const FAnimDatabaseFrameAttribute> FrameAttributes, const int32 SequenceIdx, const float SequenceTime, const FFrameRate& FrameRate)
+bool UAnimDatabaseFrameAttributeLibrary::FrameAttributeSampleAngleDegrees(float& OutValue, const FAnimDatabaseFrameAttribute& FrameAttribute, const int32 SequenceIdx, const float SequenceTime, const FFrameRate& FrameRate)
 {
-	check(OutLocations.Num() == FrameAttributes.Num());
-	for (int32 Idx = 0; Idx < FrameAttributes.Num(); Idx++)
+	if (!FrameAttribute.IsValid())
 	{
-		if (!FrameAttributeSampleLocation(OutLocations[Idx], FrameAttributes[Idx], SequenceIdx, SequenceTime, FrameRate))
-		{
-			UE::Learning::Array::Set<1, FVector>(OutLocations, FVector::ZeroVector);
-			return false;
-		}
+		UE_LOGFMT(LogAnimDatabase, Error, "FrameAttributeSampleAngleDegrees: Invalid FrameAttribute.");
+		OutValue = 0.0f;
+		return false;
 	}
 
-	return true;
-}
-
-bool UAnimDatabaseFrameAttributeLibrary::FrameAttributeSampleDirections(TArray<FVector>& OutDirections, const TArray<FAnimDatabaseFrameAttribute>& FrameAttributes, const int32 SequenceIdx, const float SequenceTime, const FFrameRate& FrameRate)
-{
-	OutDirections.SetNum(FrameAttributes.Num());
-	return FrameAttributeSampleDirectionsToArrayView(OutDirections, FrameAttributes, SequenceIdx, SequenceTime, FrameRate);
-}
+	if (FrameAttribute.Type != EAnimDatabaseAttributeType::Angle)
+	{
+		UE_LOGFMT(LogAnimDatabase, Error, "FrameAttributeSampleAngleDegrees: FrameAttribute Types don't match. Got {Type}. Expected Angle.", AttributeTypeNameInternal(FrameAttribute.Type));
+		OutValue = 0.0f;
+		return false;
+	}
 
-bool UAnimDatabaseFrameAttributeLibrary::FrameAttributeSampleDirectionsToArrayView(const TArrayView<FVector> OutDirections, const TArrayView<const FAnimDatabaseFrameAttribute> FrameAttributes, const int32 SequenceIdx, const float SequenceTime, const FFrameRate& FrameRate)
-{
-	check(OutDirections.Num() == FrameAttributes.Num());
-	for (int32 Idx = 0; Idx < FrameAttributes.Num(); Idx++)
+	int32 EntryIdx = INDEX_NONE;
+	int32 RangeIdx = INDEX_NONE;
+	float RangeTime = 1.0f;
+	if (!FrameAttribute.FrameAttribute->FrameRangeSet.FindTime(EntryIdx, RangeIdx, RangeTime, SequenceIdx, SequenceTime, 1.0f / FMath::Max(FrameRate.AsDecimal(), UE_SMALL_NUMBER)))
 	{
-		if (!FrameAttributeSampleDirection(OutDirections[Idx], FrameAttributes[Idx], SequenceIdx, SequenceTime, FrameRate))
-		{
-			UE::Learning::Array::Set<1, FVector>(OutDirections, FVector::ForwardVector);
-			return false;
-		}
+		OutValue = 0.0f;
+		return false;
 	}
 
-	return true;
-}
+	const int32 RangeOffset = FrameAttribute.FrameAttribute->FrameRangeSet.GetEntryRangeOffset(EntryIdx, RangeIdx);
+	const int32 RangeLength = FrameAttribute.FrameAttribute->FrameRangeSet.GetEntryRangeLength(EntryIdx, RangeIdx);
+	const float RangeFrameTime = 1.0f / FMath::Max(FrameRate.AsDecimal(), UE_SMALL_NUMBER);
 
-bool UAnimDatabaseFrameAttributeLibrary::FrameAttributeSampleLinearVelocities(TArray<FVector>& OutLinearVelocities, const TArray<FAnimDatabaseFrameAttribute>& FrameAttributes, const int32 SequenceIdx, const float SequenceTime, const FFrameRate& FrameRate)
-{
-	OutLinearVelocities.SetNum(FrameAttributes.Num());
-	return FrameAttributeSampleLinearVelocitiesToArrayView(OutLinearVelocities, FrameAttributes, SequenceIdx, SequenceTime, FrameRate);
-}
+	int64 SampleFrame0, SampleFrame1, SampleFrame2, SampleFrame3;
+	float SampleAlpha;
+	UE::AnimDatabase::Math::ComputeCubicSampleFramesAndAlpha(
+		SampleFrame0,
+		SampleFrame1,
+		SampleFrame2,
+		SampleFrame3,
+		SampleAlpha,
+		RangeTime / RangeFrameTime,
+		RangeLength);
 
-bool UAnimDatabaseFrameAttributeLibrary::FrameAttributeSampleLinearVelocitiesToArrayView(const TArrayView<FVector> OutLinearVelocities, const TArrayView<const FAnimDatabaseFrameAttribute> FrameAttributes, const int32 SequenceIdx, const float SequenceTime, const FFrameRate& FrameRate)
-{
-	check(OutLinearVelocities.Num() == FrameAttributes.Num());
-	for (int32 Idx = 0; Idx < FrameAttributes.Num(); Idx++)
+	if (SampleFrame0 == SampleFrame1 && SampleFrame2 == SampleFrame3)
 	{
-		if (!FrameAttributeSampleLinearVelocity(OutLinearVelocities[Idx], FrameAttributes[Idx], SequenceIdx, SequenceTime, FrameRate))
-		{
-			UE::Learning::Array::Set<1, FVector>(OutLinearVelocities, FVector::ZeroVector);
-			return false;
-		}
+		UE::AnimDatabase::Math::AngleInterpolateLinear(
+			OutValue,
+			FrameAttribute.FrameAttribute->AttributeData[0][RangeOffset + SampleFrame1],
+			FrameAttribute.FrameAttribute->AttributeData[0][RangeOffset + SampleFrame2],
+			SampleAlpha,
+			RangeFrameTime);
+		OutValue = FMath::RadiansToDegrees(OutValue);
+	}
+	else if (SampleFrame0 == SampleFrame1)
+	{
+		UE::AnimDatabase::Math::AngleInterpolateCubicMonoStart(
+			OutValue,
+			FrameAttribute.FrameAttribute->AttributeData[0][RangeOffset + SampleFrame1],
+			FrameAttribute.FrameAttribute->AttributeData[0][RangeOffset + SampleFrame2],
+			FrameAttribute.FrameAttribute->AttributeData[0][RangeOffset + SampleFrame3],
+			SampleAlpha,
+			RangeFrameTime);
+		OutValue = FMath::RadiansToDegrees(OutValue);
+	}
+	else if (SampleFrame2 == SampleFrame3)
+	{
+		UE::AnimDatabase::Math::AngleInterpolateCubicMonoEnd(
+			OutValue,
+			FrameAttribute.FrameAttribute->AttributeData[0][RangeOffset + SampleFrame0],
+			FrameAttribute.FrameAttribute->AttributeData[0][RangeOffset + SampleFrame1],
+			FrameAttribute.FrameAttribute->AttributeData[0][RangeOffset + SampleFrame2],
+			SampleAlpha, RangeFrameTime);
+		OutValue = FMath::RadiansToDegrees(OutValue);
+	}
+	else
+	{
+		UE::AnimDatabase::Math::AngleInterpolateCubicMono(
+			OutValue,
+			FrameAttribute.FrameAttribute->AttributeData[0][RangeOffset + SampleFrame0],
+			FrameAttribute.FrameAttribute->AttributeData[0][RangeOffset + SampleFrame1],
+			FrameAttribute.FrameAttribute->AttributeData[0][RangeOffset + SampleFrame2],
+			FrameAttribute.FrameAttribute->AttributeData[0][RangeOffset + SampleFrame3],
+			SampleAlpha, RangeFrameTime);
+		OutValue = FMath::RadiansToDegrees(OutValue);
 	}
 
 	return true;
 }
 
-int32 UAnimDatabaseFrameAttributeLibrary::FrameAttributeTotalRangeNum(const FAnimDatabaseFrameAttribute& FrameAttribute)
+bool UAnimDatabaseFrameAttributeLibrary::FrameAttributeSampleAngleRadians(float& OutValue, const FAnimDatabaseFrameAttribute& FrameAttribute, const int32 SequenceIdx, const float SequenceTime, const FFrameRate& FrameRate)
 {
-	return UAnimDatabaseFrameRangesLibrary::FrameRangesTotalRangeNum(FrameAttributeFrameRanges(FrameAttribute));
-}
-
+	if (!FrameAttribute.IsValid())
+	{
+		UE_LOGFMT(LogAnimDatabase, Error, "FrameAttributeSampleAngleRadians: Invalid FrameAttribute.");
+		OutValue = 0.0f;
+		return false;
+	}
+
+	if (FrameAttribute.Type != EAnimDatabaseAttributeType::Angle)
+	{
+		UE_LOGFMT(LogAnimDatabase, Error, "FrameAttributeSampleAngleRadians: FrameAttribute Types don't match. Got {Type}. Expected Angle.", AttributeTypeNameInternal(FrameAttribute.Type));
+		OutValue = 0.0f;
+		return false;
+	}
+
+	int32 EntryIdx = INDEX_NONE;
+	int32 RangeIdx = INDEX_NONE;
+	float RangeTime = 1.0f;
+	if (!FrameAttribute.FrameAttribute->FrameRangeSet.FindTime(EntryIdx, RangeIdx, RangeTime, SequenceIdx, SequenceTime, 1.0f / FMath::Max(FrameRate.AsDecimal(), UE_SMALL_NUMBER)))
+	{
+		OutValue = 0.0f;
+		return false;
+	}
+
+	const int32 RangeOffset = FrameAttribute.FrameAttribute->FrameRangeSet.GetEntryRangeOffset(EntryIdx, RangeIdx);
+	const int32 RangeLength = FrameAttribute.FrameAttribute->FrameRangeSet.GetEntryRangeLength(EntryIdx, RangeIdx);
+	const float RangeFrameTime = 1.0f / FMath::Max(FrameRate.AsDecimal(), UE_SMALL_NUMBER);
+
+	int64 SampleFrame0, SampleFrame1, SampleFrame2, SampleFrame3;
+	float SampleAlpha;
+	UE::AnimDatabase::Math::ComputeCubicSampleFramesAndAlpha(
+		SampleFrame0,
+		SampleFrame1,
+		SampleFrame2,
+		SampleFrame3,
+		SampleAlpha,
+		RangeTime / RangeFrameTime,
+		RangeLength);
+
+	if (SampleFrame0 == SampleFrame1 && SampleFrame2 == SampleFrame3)
+	{
+		UE::AnimDatabase::Math::AngleInterpolateLinear(
+			OutValue,
+			FrameAttribute.FrameAttribute->AttributeData[0][RangeOffset + SampleFrame1],
+			FrameAttribute.FrameAttribute->AttributeData[0][RangeOffset + SampleFrame2],
+			SampleAlpha,
+			RangeFrameTime);
+	}
+	else if (SampleFrame0 == SampleFrame1)
+	{
+		UE::AnimDatabase::Math::AngleInterpolateCubicMonoStart(
+			OutValue,
+			FrameAttribute.FrameAttribute->AttributeData[0][RangeOffset + SampleFrame1],
+			FrameAttribute.FrameAttribute->AttributeData[0][RangeOffset + SampleFrame2],
+			FrameAttribute.FrameAttribute->AttributeData[0][RangeOffset + SampleFrame3],
+			SampleAlpha,
+			RangeFrameTime);
+	}
+	else if (SampleFrame2 == SampleFrame3)
+	{
+		UE::AnimDatabase::Math::AngleInterpolateCubicMonoEnd(
+			OutValue,
+			FrameAttribute.FrameAttribute->AttributeData[0][RangeOffset + SampleFrame0],
+			FrameAttribute.FrameAttribute->AttributeData[0][RangeOffset + SampleFrame1],
+			FrameAttribute.FrameAttribute->AttributeData[0][RangeOffset + SampleFrame2],
+			SampleAlpha, RangeFrameTime);
+	}
+	else
+	{
+		UE::AnimDatabase::Math::AngleInterpolateCubicMono(
+			OutValue,
+			FrameAttribute.FrameAttribute->AttributeData[0][RangeOffset + SampleFrame0],
+			FrameAttribute.FrameAttribute->AttributeData[0][RangeOffset + SampleFrame1],
+			FrameAttribute.FrameAttribute->AttributeData[0][RangeOffset + SampleFrame2],
+			FrameAttribute.FrameAttribute->AttributeData[0][RangeOffset + SampleFrame3],
+			SampleAlpha, RangeFrameTime);
+	}
+
+	return true;
+}
+
+bool UAnimDatabaseFrameAttributeLibrary::FrameAttributeSampleLocations(TArray<FVector>& OutLocations, const TArray<FAnimDatabaseFrameAttribute>& FrameAttributes, const int32 SequenceIdx, const float SequenceTime, const FFrameRate& FrameRate)
+{
+	OutLocations.SetNum(FrameAttributes.Num());
+	return FrameAttributeSampleLocationsToArrayView(OutLocations, FrameAttributes, SequenceIdx, SequenceTime, FrameRate);
+}
+
+bool UAnimDatabaseFrameAttributeLibrary::FrameAttributeSampleLocationsToArrayView(const TArrayView<FVector> OutLocations, const TArrayView<const FAnimDatabaseFrameAttribute> FrameAttributes, const int32 SequenceIdx, const float SequenceTime, const FFrameRate& FrameRate)
+{
+	check(OutLocations.Num() == FrameAttributes.Num());
+	for (int32 Idx = 0; Idx < FrameAttributes.Num(); Idx++)
+	{
+		if (!FrameAttributeSampleLocation(OutLocations[Idx], FrameAttributes[Idx], SequenceIdx, SequenceTime, FrameRate))
+		{
+			UE::Learning::Array::Set<1, FVector>(OutLocations, FVector::ZeroVector);
+			return false;
+		}
+	}
+
+	return true;
+}
+
+bool UAnimDatabaseFrameAttributeLibrary::FrameAttributeSampleDirections(TArray<FVector>& OutDirections, const TArray<FAnimDatabaseFrameAttribute>& FrameAttributes, const int32 SequenceIdx, const float SequenceTime, const FFrameRate& FrameRate)
+{
+	OutDirections.SetNum(FrameAttributes.Num());
+	return FrameAttributeSampleDirectionsToArrayView(OutDirections, FrameAttributes, SequenceIdx, SequenceTime, FrameRate);
+}
+
+bool UAnimDatabaseFrameAttributeLibrary::FrameAttributeSampleDirectionsToArrayView(const TArrayView<FVector> OutDirections, const TArrayView<const FAnimDatabaseFrameAttribute> FrameAttributes, const int32 SequenceIdx, const float SequenceTime, const FFrameRate& FrameRate)
+{
+	check(OutDirections.Num() == FrameAttributes.Num());
+	for (int32 Idx = 0; Idx < FrameAttributes.Num(); Idx++)
+	{
+		if (!FrameAttributeSampleDirection(OutDirections[Idx], FrameAttributes[Idx], SequenceIdx, SequenceTime, FrameRate))
+		{
+			UE::Learning::Array::Set<1, FVector>(OutDirections, FVector::ForwardVector);
+			return false;
+		}
+	}
+
+	return true;
+}
+
+bool UAnimDatabaseFrameAttributeLibrary::FrameAttributeSampleLinearVelocities(TArray<FVector>& OutLinearVelocities, const TArray<FAnimDatabaseFrameAttribute>& FrameAttributes, const int32 SequenceIdx, const float SequenceTime, const FFrameRate& FrameRate)
+{
+	OutLinearVelocities.SetNum(FrameAttributes.Num());
+	return FrameAttributeSampleLinearVelocitiesToArrayView(OutLinearVelocities, FrameAttributes, SequenceIdx, SequenceTime, FrameRate);
+}
+
+bool UAnimDatabaseFrameAttributeLibrary::FrameAttributeSampleLinearVelocitiesToArrayView(const TArrayView<FVector> OutLinearVelocities, const TArrayView<const FAnimDatabaseFrameAttribute> FrameAttributes, const int32 SequenceIdx, const float SequenceTime, const FFrameRate& FrameRate)
+{
+	check(OutLinearVelocities.Num() == FrameAttributes.Num());
+	for (int32 Idx = 0; Idx < FrameAttributes.Num(); Idx++)
+	{
+		if (!FrameAttributeSampleLinearVelocity(OutLinearVelocities[Idx], FrameAttributes[Idx], SequenceIdx, SequenceTime, FrameRate))
+		{
+			UE::Learning::Array::Set<1, FVector>(OutLinearVelocities, FVector::ZeroVector);
+			return false;
+		}
+	}
+
+	return true;
+}
+
+int32 UAnimDatabaseFrameAttributeLibrary::FrameAttributeTotalRangeNum(const FAnimDatabaseFrameAttribute& FrameAttribute)
+{
+	return UAnimDatabaseFrameRangesLibrary::FrameRangesTotalRangeNum(FrameAttributeFrameRanges(FrameAttribute));
+}
+
 int32 UAnimDatabaseFrameAttributeLibrary::FrameAttributeTotalFrameNum(const FAnimDatabaseFrameAttribute& FrameAttribute)
 {
 	return UAnimDatabaseFrameRangesLibrary::FrameRangesTotalFrameNum(FrameAttributeFrameRanges(FrameAttribute));
@@ -5466,7 +5876,7 @@ void UAnimDatabaseFrameAttributeLibrary::FrameAttributeToBoolArrayView(const TAr
 
 	if (FrameAttribute.Type != EAnimDatabaseAttributeType::Bool)
 	{
-		UE_LOGFMT(LogAnimDatabase, Error, "FrameAttributeToBoolArrayView: Incorrect FrameAttribute Type. Expected Bool, got {Type}.", *UEnum::GetValueAsString(FrameAttribute.Type));
+		UE_LOGFMT(LogAnimDatabase, Error, "FrameAttributeToBoolArrayView: Incorrect FrameAttribute Type. Expected Bool, got {Type}.", AttributeTypeNameInternal(FrameAttribute.Type));
 		return;
 	}
 
@@ -5495,7 +5905,7 @@ void UAnimDatabaseFrameAttributeLibrary::FrameAttributeToFloatArrayView(const TA
 
 	if (FrameAttribute.Type != EAnimDatabaseAttributeType::Float)
 	{
-		UE_LOGFMT(LogAnimDatabase, Error, "FrameAttributeToFloatArrayView: Incorrect FrameAttribute Type. Expected Float, got {Type}.", *UEnum::GetValueAsString(FrameAttribute.Type));
+		UE_LOGFMT(LogAnimDatabase, Error, "FrameAttributeToFloatArrayView: Incorrect FrameAttribute Type. Expected Float, got {Type}.", AttributeTypeNameInternal(FrameAttribute.Type));
 		return;
 	}
 
@@ -5529,7 +5939,7 @@ void UAnimDatabaseFrameAttributeLibrary::FrameAttributeToVectorArrayView(TArray<
 		FrameAttribute.Type != EAnimDatabaseAttributeType::AngularVelocity &&
 		FrameAttribute.Type != EAnimDatabaseAttributeType::ScalarVelocity)
 	{
-		UE_LOGFMT(LogAnimDatabase, Error, "FrameAttributeToVectorArrayView: Incorrect FrameAttribute Type. Expected Location, Direction, Scale or Velocity, got {Type}.", *UEnum::GetValueAsString(FrameAttribute.Type));
+		UE_LOGFMT(LogAnimDatabase, Error, "FrameAttributeToVectorArrayView: Incorrect FrameAttribute Type. Expected Location, Direction, Scale or Velocity, got {Type}.", AttributeTypeNameInternal(FrameAttribute.Type));
 		return;
 	}
 
@@ -5575,6 +5985,11 @@ void UAnimDatabaseFrameAttributeLibrary::FrameAttributeToVectorArrayView(TArray<
 		}
 	}
 	break;
+
+	default:
+	{
+		checkNoEntry();
+	}
 	}
 
 }
@@ -5595,7 +6010,7 @@ void UAnimDatabaseFrameAttributeLibrary::FrameAttributeToQuatArrayView(TArray<FQ
 
 	if (FrameAttribute.Type != EAnimDatabaseAttributeType::Rotation)
 	{
-		UE_LOGFMT(LogAnimDatabase, Error, "FrameAttributeToQuatArray: Incorrect FrameAttribute Type. Expected Rotation, got {Type}.", *UEnum::GetValueAsString(FrameAttribute.Type));
+		UE_LOGFMT(LogAnimDatabase, Error, "FrameAttributeToQuatArray: Incorrect FrameAttribute Type. Expected Rotation, got {Type}.", AttributeTypeNameInternal(FrameAttribute.Type));
 		return;
 	}
 
@@ -5624,7 +6039,7 @@ void UAnimDatabaseFrameAttributeLibrary::FrameAttributeToRotatorArrayView(TArray
 
 	if (FrameAttribute.Type != EAnimDatabaseAttributeType::Rotation)
 	{
-		UE_LOGFMT(LogAnimDatabase, Error, "FrameAttributeToRotatorArrayView: Incorrect FrameAttribute Type. Expected Rotation, got {Type}.", *UEnum::GetValueAsString(FrameAttribute.Type));
+		UE_LOGFMT(LogAnimDatabase, Error, "FrameAttributeToRotatorArrayView: Incorrect FrameAttribute Type. Expected Rotation, got {Type}.", AttributeTypeNameInternal(FrameAttribute.Type));
 		return;
 	}
 
@@ -5653,7 +6068,7 @@ void UAnimDatabaseFrameAttributeLibrary::FrameAttributeToTransformArrayView(TArr
 
 	if (FrameAttribute.Type != EAnimDatabaseAttributeType::Transform)
 	{
-		UE_LOGFMT(LogAnimDatabase, Error, "FrameAttributeToTransformArrayView: Incorrect FrameAttribute Type. Expected Transform, got {Type}.", *UEnum::GetValueAsString(FrameAttribute.Type));
+		UE_LOGFMT(LogAnimDatabase, Error, "FrameAttributeToTransformArrayView: Incorrect FrameAttribute Type. Expected Transform, got {Type}.", AttributeTypeNameInternal(FrameAttribute.Type));
 		return;
 	}
 
@@ -5666,6 +6081,64 @@ void UAnimDatabaseFrameAttributeLibrary::FrameAttributeToTransformArrayView(TArr
 	}
 }
 
+void UAnimDatabaseFrameAttributeLibrary::FrameAttributeToAngleArrayDegrees(TArray<float>& OutArray, const FAnimDatabaseFrameAttribute& FrameAttribute)
+{
+	OutArray.SetNumZeroed(FrameAttributeTotalFrameNum(FrameAttribute));
+	FrameAttributeToAngleArrayViewDegrees(OutArray, FrameAttribute);
+}
+
+void UAnimDatabaseFrameAttributeLibrary::FrameAttributeToAngleArrayViewDegrees(const TArrayView<float> OutArray, const FAnimDatabaseFrameAttribute& FrameAttribute)
+{
+	if (!FrameAttribute.IsValid())
+	{
+		UE_LOGFMT(LogAnimDatabase, Error, "FrameAttributeToAngleArrayViewDegrees: Invalid FrameAttribute.");
+		return;
+	}
+
+	if (FrameAttribute.Type != EAnimDatabaseAttributeType::Angle)
+	{
+		UE_LOGFMT(LogAnimDatabase, Error, "FrameAttributeToAngleArrayViewDegrees: Incorrect FrameAttribute Type. Expected Angle, got {Type}.", AttributeTypeNameInternal(FrameAttribute.Type));
+		return;
+	}
+
+	const int32 FrameNum = FrameAttributeTotalFrameNum(FrameAttribute);
+	check(OutArray.Num() == FrameNum);
+
+	for (int32 FrameIdx = 0; FrameIdx < FrameNum; FrameIdx++)
+	{
+		OutArray[FrameIdx] = FrameAttribute.GetAsAngleDegrees(FrameIdx);
+	}
+}
+
+void UAnimDatabaseFrameAttributeLibrary::FrameAttributeToAngleArrayRadians(TArray<float>& OutArray, const FAnimDatabaseFrameAttribute& FrameAttribute)
+{
+	OutArray.SetNumZeroed(FrameAttributeTotalFrameNum(FrameAttribute));
+	FrameAttributeToAngleArrayViewRadians(OutArray, FrameAttribute);
+}
+
+void UAnimDatabaseFrameAttributeLibrary::FrameAttributeToAngleArrayViewRadians(const TArrayView<float> OutArray, const FAnimDatabaseFrameAttribute& FrameAttribute)
+{
+	if (!FrameAttribute.IsValid())
+	{
+		UE_LOGFMT(LogAnimDatabase, Error, "FrameAttributeToAngleArrayViewRadians: Invalid FrameAttribute.");
+		return;
+	}
+
+	if (FrameAttribute.Type != EAnimDatabaseAttributeType::Angle)
+	{
+		UE_LOGFMT(LogAnimDatabase, Error, "FrameAttributeToAngleArrayViewRadians: Incorrect FrameAttribute Type. Expected Angle, got {Type}.", AttributeTypeNameInternal(FrameAttribute.Type));
+		return;
+	}
+
+	const int32 FrameNum = FrameAttributeTotalFrameNum(FrameAttribute);
+	check(OutArray.Num() == FrameNum);
+
+	for (int32 FrameIdx = 0; FrameIdx < FrameNum; FrameIdx++)
+	{
+		OutArray[FrameIdx] = FrameAttribute.GetAsAngleRadians(FrameIdx);
+	}
+}
+
 FAnimDatabaseFrameAttribute UAnimDatabaseFrameAttributeLibrary::FrameAttributeLocationMatchingDistance(
 	const FAnimDatabaseFrameAttribute& LocationFrameAttribute,
 	const FVector Location,
@@ -5682,7 +6155,7 @@ FAnimDatabaseFrameAttribute UAnimDatabaseFrameAttributeLibrary::FrameAttributeLo
 
 	if (LocationFrameAttribute.Type != EAnimDatabaseAttributeType::Location)
 	{
-		UE_LOGFMT(LogAnimDatabase, Error, "FrameAttributeLocationMatchingDistance: Incorrect FrameAttribute Type. Expected Location, got {Type}.", *UEnum::GetValueAsString(LocationFrameAttribute.Type));
+		UE_LOGFMT(LogAnimDatabase, Error, "FrameAttributeLocationMatchingDistance: Incorrect FrameAttribute Type. Expected Location, got {Type}.", AttributeTypeNameInternal(LocationFrameAttribute.Type));
 		return FAnimDatabaseFrameAttribute();
 	}
 
@@ -5714,7 +6187,7 @@ FAnimDatabaseFrameAttribute UAnimDatabaseFrameAttributeLibrary::FrameAttributeDi
 
 	if (DirectionFrameAttribute.Type != EAnimDatabaseAttributeType::Direction)
 	{
-		UE_LOGFMT(LogAnimDatabase, Error, "FrameAttributeDirectionMatchingDistance: Incorrect FrameAttribute Type. Expected Direction, got {Type}.", *UEnum::GetValueAsString(DirectionFrameAttribute.Type));
+		UE_LOGFMT(LogAnimDatabase, Error, "FrameAttributeDirectionMatchingDistance: Incorrect FrameAttribute Type. Expected Direction, got {Type}.", AttributeTypeNameInternal(DirectionFrameAttribute.Type));
 		return FAnimDatabaseFrameAttribute();
 	}
 
@@ -5747,7 +6220,7 @@ FAnimDatabaseFrameAttribute UAnimDatabaseFrameAttributeLibrary::FrameAttributeRo
 
 	if (RotationFrameAttribute.Type != EAnimDatabaseAttributeType::Rotation)
 	{
-		UE_LOGFMT(LogAnimDatabase, Error, "FrameAttributeRotationMatchingDistance: Incorrect FrameAttribute Type. Expected Rotation, got {Type}.", *UEnum::GetValueAsString(RotationFrameAttribute.Type));
+		UE_LOGFMT(LogAnimDatabase, Error, "FrameAttributeRotationMatchingDistance: Incorrect FrameAttribute Type. Expected Rotation, got {Type}.", AttributeTypeNameInternal(RotationFrameAttribute.Type));
 		return FAnimDatabaseFrameAttribute();
 	}
 
@@ -5782,7 +6255,7 @@ FAnimDatabaseFrameAttribute UAnimDatabaseFrameAttributeLibrary::FrameAttributeVe
 		VelocityFrameAttribute.Type != EAnimDatabaseAttributeType::AngularVelocity &&
 		VelocityFrameAttribute.Type != EAnimDatabaseAttributeType::ScalarVelocity)
 	{
-		UE_LOGFMT(LogAnimDatabase, Error, "FrameAttributeVelocityMatchingDistance: Incorrect FrameAttribute Type. Expected LinearVelocity, AngularVelocity, or ScalarVelocity, got {Type}.", *UEnum::GetValueAsString(VelocityFrameAttribute.Type));
+		UE_LOGFMT(LogAnimDatabase, Error, "FrameAttributeVelocityMatchingDistance: Incorrect FrameAttribute Type. Expected LinearVelocity, AngularVelocity, or ScalarVelocity, got {Type}.", AttributeTypeNameInternal(VelocityFrameAttribute.Type));
 		return FAnimDatabaseFrameAttribute();
 	}
 
@@ -5815,7 +6288,7 @@ FAnimDatabaseFrameAttribute UAnimDatabaseFrameAttributeLibrary::FrameAttributeEv
 
 	if (EventFrameAttribute.Type != EAnimDatabaseAttributeType::Event)
 	{
-		UE_LOGFMT(LogAnimDatabase, Error, "FrameAttributeEventMatchingDistance: Incorrect FrameAttribute Type. Expected Event, got {Type}.", *UEnum::GetValueAsString(EventFrameAttribute.Type));
+		UE_LOGFMT(LogAnimDatabase, Error, "FrameAttributeEventMatchingDistance: Incorrect FrameAttribute Type. Expected Event, got {Type}.", AttributeTypeNameInternal(EventFrameAttribute.Type));
 		return FAnimDatabaseFrameAttribute();
 	}
 
@@ -5864,7 +6337,7 @@ FAnimDatabaseFrameAttribute UAnimDatabaseFrameAttributeLibrary::FrameAttributeIn
 	if (LocationFrameAttribute.Type != EAnimDatabaseAttributeType::Location ||
 		VelocityFrameAttribute.Type != EAnimDatabaseAttributeType::LinearVelocity)
 	{
-		UE_LOGFMT(LogAnimDatabase, Error, "FrameAttributeInertializationMatchingDistance: Incorrect FrameAttribute Type. Expected Location and LinearVelocity, got {Type0} and {Type1}.", *UEnum::GetValueAsString(LocationFrameAttribute.Type), *UEnum::GetValueAsString(VelocityFrameAttribute.Type));
+		UE_LOGFMT(LogAnimDatabase, Error, "FrameAttributeInertializationMatchingDistance: Incorrect FrameAttribute Type. Expected Location and LinearVelocity, got {Type0} and {Type1}.", AttributeTypeNameInternal(LocationFrameAttribute.Type), AttributeTypeNameInternal(VelocityFrameAttribute.Type));
 		return FAnimDatabaseFrameAttribute();
 	}
 
@@ -5904,7 +6377,7 @@ float UAnimDatabaseFrameAttributeLibrary::FrameAttributeAverageStd(const FAnimDa
 		FrameAttribute.Type != EAnimDatabaseAttributeType::AngularVelocity &&
 		FrameAttribute.Type != EAnimDatabaseAttributeType::ScalarVelocity)
 	{
-		UE_LOGFMT(LogAnimDatabase, Error, "FrameAttributeAverageStd: Incorrect FrameAttribute Type. Expected Float, Location, LinearVelocity, AngularVelocity, or ScalarVelocity, got {Type}.", *UEnum::GetValueAsString(FrameAttribute.Type));
+		UE_LOGFMT(LogAnimDatabase, Error, "FrameAttributeAverageStd: Incorrect FrameAttribute Type. Expected Float, Location, LinearVelocity, AngularVelocity, or ScalarVelocity, got {Type}.", AttributeTypeNameInternal(FrameAttribute.Type));
 		return 0.0f;
 	}
 
@@ -5989,7 +6462,7 @@ void UAnimDatabaseFrameAttributeLibrary::FindFrameAttributeMinimum(
 		OutSequenceIdx = INDEX_NONE;
 		OutSequenceFrame = INDEX_NONE;
 		OutMinimumValue = 0.0f;
-		UE_LOGFMT(LogAnimDatabase, Error, "FindFrameAttributeMinimum: Incorrect FrameAttribute Type. Expected Float, got {Type}.", *UEnum::GetValueAsString(FrameAttribute.Type));
+		UE_LOGFMT(LogAnimDatabase, Error, "FindFrameAttributeMinimum: Incorrect FrameAttribute Type. Expected Float, got {Type}.", AttributeTypeNameInternal(FrameAttribute.Type));
 		return;
 	}
 
@@ -6080,7 +6553,7 @@ void UAnimDatabaseFrameAttributeLibrary::FindFrameAttributeMaximum(
 		OutSequenceIdx = INDEX_NONE;
 		OutSequenceFrame = INDEX_NONE;
 		OutMaximumValue = 0.0f;
-		UE_LOGFMT(LogAnimDatabase, Error, "FindFrameAttributeMaximum: Incorrect FrameAttribute Type. Expected Float, got {FrameAttribute}.", *UEnum::GetValueAsString(FrameAttribute.Type));
+		UE_LOGFMT(LogAnimDatabase, Error, "FindFrameAttributeMaximum: Incorrect FrameAttribute Type. Expected Float, got {FrameAttribute}.", AttributeTypeNameInternal(FrameAttribute.Type));
 		return;
 	}
 
@@ -6144,21 +6617,21 @@ FAnimDatabaseFrameAttribute UAnimDatabaseFrameAttributeLibrary::FrameAttributeAd
 		(A.Type == EAnimDatabaseAttributeType::Float && B.Type == EAnimDatabaseAttributeType::Float) ||
 		(A.Type == EAnimDatabaseAttributeType::Bool && B.Type == EAnimDatabaseAttributeType::Float) ||
 		(A.Type == EAnimDatabaseAttributeType::Float && B.Type == EAnimDatabaseAttributeType::Bool) ||
-		(A.Type == EAnimDatabaseAttributeType::Location && B.Type == EAnimDatabaseAttributeType::Location)))
+		(A.Type == EAnimDatabaseAttributeType::Location && B.Type == EAnimDatabaseAttributeType::Location) ||
+		(A.Type == EAnimDatabaseAttributeType::Angle && B.Type == EAnimDatabaseAttributeType::Angle)))
 	{
-		UE_LOGFMT(LogAnimDatabase, Error, "FrameAttributeAdd: Invalid FrameAttribute Types. Got {Type0} and {Type1}.", *UEnum::GetValueAsString(A.Type), *UEnum::GetValueAsString(B.Type));
+		UE_LOGFMT(LogAnimDatabase, Error, "FrameAttributeAdd: Invalid FrameAttribute Types. Got {Type0} and {Type1}.", AttributeTypeNameInternal(A.Type), AttributeTypeNameInternal(B.Type));
 		return FAnimDatabaseFrameAttribute();
 	}
 
 	EAnimDatabaseAttributeType OutputType = EAnimDatabaseAttributeType::Float;
-	if (A.Type == EAnimDatabaseAttributeType::Location && B.Type == EAnimDatabaseAttributeType::Location)
-	{
-		OutputType = EAnimDatabaseAttributeType::Location;
-	}
+	if (A.Type == EAnimDatabaseAttributeType::Location && B.Type == EAnimDatabaseAttributeType::Location) { OutputType = EAnimDatabaseAttributeType::Location; }
+	if (A.Type == EAnimDatabaseAttributeType::Angle && B.Type == EAnimDatabaseAttributeType::Angle) { OutputType = EAnimDatabaseAttributeType::Angle; }
 
 	FAnimDatabaseFrameAttribute Out = MakeEmptyFrameAttribute();
 	Out.Type = OutputType;
 	UE::Learning::FrameAttribute::Add(*Out.FrameAttribute, *A.FrameAttribute, *B.FrameAttribute);
+	if (Out.Type == EAnimDatabaseAttributeType::Angle) { UE::Learning::FrameAttribute::WrapAngleInplace(*Out.FrameAttribute); }
 	return Out;
 }
 
@@ -6174,21 +6647,21 @@ FAnimDatabaseFrameAttribute UAnimDatabaseFrameAttributeLibrary::FrameAttributeSu
 		(A.Type == EAnimDatabaseAttributeType::Float && B.Type == EAnimDatabaseAttributeType::Float) ||
 		(A.Type == EAnimDatabaseAttributeType::Bool && B.Type == EAnimDatabaseAttributeType::Float) ||
 		(A.Type == EAnimDatabaseAttributeType::Float && B.Type == EAnimDatabaseAttributeType::Bool) ||
-		(A.Type == EAnimDatabaseAttributeType::Location && B.Type == EAnimDatabaseAttributeType::Location)))
+		(A.Type == EAnimDatabaseAttributeType::Location && B.Type == EAnimDatabaseAttributeType::Location) ||
+		(A.Type == EAnimDatabaseAttributeType::Angle && B.Type == EAnimDatabaseAttributeType::Angle)))
 	{
-		UE_LOGFMT(LogAnimDatabase, Error, "FrameAttributeSubtract: Invalid FrameAttribute Types. Got {Type0} and {Type1}.", *UEnum::GetValueAsString(A.Type), *UEnum::GetValueAsString(B.Type));
+		UE_LOGFMT(LogAnimDatabase, Error, "FrameAttributeSubtract: Invalid FrameAttribute Types. Got {Type0} and {Type1}.", AttributeTypeNameInternal(A.Type), AttributeTypeNameInternal(B.Type));
 		return FAnimDatabaseFrameAttribute();
 	}
 
 	EAnimDatabaseAttributeType OutputType = EAnimDatabaseAttributeType::Float;
-	if (A.Type == EAnimDatabaseAttributeType::Location && B.Type == EAnimDatabaseAttributeType::Location)
-	{
-		OutputType = EAnimDatabaseAttributeType::Location;
-	}
+	if (A.Type == EAnimDatabaseAttributeType::Location && B.Type == EAnimDatabaseAttributeType::Location) { OutputType = EAnimDatabaseAttributeType::Location; }
+	if (A.Type == EAnimDatabaseAttributeType::Angle && B.Type == EAnimDatabaseAttributeType::Angle) { OutputType = EAnimDatabaseAttributeType::Angle; }
 
 	FAnimDatabaseFrameAttribute Out = MakeEmptyFrameAttribute();
 	Out.Type = OutputType;
 	UE::Learning::FrameAttribute::Sub(*Out.FrameAttribute, *A.FrameAttribute, *B.FrameAttribute);
+	if (Out.Type == EAnimDatabaseAttributeType::Angle) { UE::Learning::FrameAttribute::WrapAngleInplace(*Out.FrameAttribute); }
 	return Out;
 }
 
@@ -6205,9 +6678,11 @@ FAnimDatabaseFrameAttribute UAnimDatabaseFrameAttributeLibrary::FrameAttributeMu
 		(A.Type == EAnimDatabaseAttributeType::Float && B.Type == EAnimDatabaseAttributeType::Float) ||
 		(A.Type == EAnimDatabaseAttributeType::Location && B.Type == EAnimDatabaseAttributeType::Location) ||
 		(A.Type == EAnimDatabaseAttributeType::Rotation && B.Type == EAnimDatabaseAttributeType::Rotation) ||
-		(A.Type == EAnimDatabaseAttributeType::Transform && B.Type == EAnimDatabaseAttributeType::Transform)))
+		(A.Type == EAnimDatabaseAttributeType::Transform && B.Type == EAnimDatabaseAttributeType::Transform) ||
+		(A.Type == EAnimDatabaseAttributeType::Angle && B.Type == EAnimDatabaseAttributeType::Float) ||
+		(A.Type == EAnimDatabaseAttributeType::Float && B.Type == EAnimDatabaseAttributeType::Angle)))
 	{
-		UE_LOGFMT(LogAnimDatabase, Error, "FrameAttributeMultiply: Invalid FrameAttribute Types. Got {Type0} and {Type1}.", *UEnum::GetValueAsString(A.Type), *UEnum::GetValueAsString(B.Type));
+		UE_LOGFMT(LogAnimDatabase, Error, "FrameAttributeMultiply: Invalid FrameAttribute Types. Got {Type0} and {Type1}.", AttributeTypeNameInternal(A.Type), AttributeTypeNameInternal(B.Type));
 		return FAnimDatabaseFrameAttribute();
 	}
 
@@ -6233,6 +6708,14 @@ FAnimDatabaseFrameAttribute UAnimDatabaseFrameAttributeLibrary::FrameAttributeMu
 		UE::Learning::FrameAttribute::TransformMul(*Out.FrameAttribute, *A.FrameAttribute, *B.FrameAttribute);
 		return Out;
 	}
+	else if (A.Type == EAnimDatabaseAttributeType::Angle || B.Type == EAnimDatabaseAttributeType::Angle)
+	{
+		FAnimDatabaseFrameAttribute Out = MakeEmptyFrameAttribute();
+		Out.Type = EAnimDatabaseAttributeType::Angle;
+		UE::Learning::FrameAttribute::Mul(*Out.FrameAttribute, *A.FrameAttribute, *B.FrameAttribute);
+		UE::Learning::FrameAttribute::WrapAngleInplace(*Out.FrameAttribute);
+		return Out;
+	}
 	else
 	{
 		checkNoEntry();
@@ -6252,9 +6735,11 @@ FAnimDatabaseFrameAttribute UAnimDatabaseFrameAttributeLibrary::FrameAttributeDi
 		(A.Type == EAnimDatabaseAttributeType::Float && B.Type == EAnimDatabaseAttributeType::Float) ||
 		(A.Type == EAnimDatabaseAttributeType::Location && B.Type == EAnimDatabaseAttributeType::Location) ||
 		(A.Type == EAnimDatabaseAttributeType::Rotation && B.Type == EAnimDatabaseAttributeType::Rotation) ||
-		(A.Type == EAnimDatabaseAttributeType::Transform && B.Type == EAnimDatabaseAttributeType::Transform)))
+		(A.Type == EAnimDatabaseAttributeType::Transform && B.Type == EAnimDatabaseAttributeType::Transform) ||
+		(A.Type == EAnimDatabaseAttributeType::Angle && B.Type == EAnimDatabaseAttributeType::Float) ||
+		(A.Type == EAnimDatabaseAttributeType::Float && B.Type == EAnimDatabaseAttributeType::Angle)))
 	{
-		UE_LOGFMT(LogAnimDatabase, Error, "FrameAttributeDivide: Invalid FrameAttribute Types. Got {Type0} and {Type1}.", *UEnum::GetValueAsString(A.Type), *UEnum::GetValueAsString(B.Type));
+		UE_LOGFMT(LogAnimDatabase, Error, "FrameAttributeDivide: Invalid FrameAttribute Types. Got {Type0} and {Type1}.", AttributeTypeNameInternal(A.Type), AttributeTypeNameInternal(B.Type));
 		return FAnimDatabaseFrameAttribute();
 	}
 
@@ -6280,6 +6765,14 @@ FAnimDatabaseFrameAttribute UAnimDatabaseFrameAttributeLibrary::FrameAttributeDi
 		UE::Learning::FrameAttribute::TransformDiv(*Out.FrameAttribute, *A.FrameAttribute, *B.FrameAttribute);
 		return Out;
 	}
+	else if (A.Type == EAnimDatabaseAttributeType::Angle || B.Type == EAnimDatabaseAttributeType::Angle)
+	{
+		FAnimDatabaseFrameAttribute Out = MakeEmptyFrameAttribute();
+		Out.Type = EAnimDatabaseAttributeType::Angle;
+		UE::Learning::FrameAttribute::Div(*Out.FrameAttribute, *A.FrameAttribute, *B.FrameAttribute);
+		UE::Learning::FrameAttribute::WrapAngleInplace(*Out.FrameAttribute);
+		return Out;
+	}
 	else
 	{
 		checkNoEntry();
@@ -6300,9 +6793,10 @@ FAnimDatabaseFrameAttribute UAnimDatabaseFrameAttributeLibrary::FrameAttributeAd
 		A.Type != EAnimDatabaseAttributeType::Scale &&
 		A.Type != EAnimDatabaseAttributeType::LinearVelocity &&
 		A.Type != EAnimDatabaseAttributeType::AngularVelocity &&
-		A.Type != EAnimDatabaseAttributeType::ScalarVelocity)
+		A.Type != EAnimDatabaseAttributeType::ScalarVelocity &&
+		A.Type != EAnimDatabaseAttributeType::Angle)
 	{
-		UE_LOGFMT(LogAnimDatabase, Error, "FrameAttributeAddFloat: FrameAttribute Type not valid. Got {Type}.", *UEnum::GetValueAsString(A.Type));
+		UE_LOGFMT(LogAnimDatabase, Error, "FrameAttributeAddFloat: FrameAttribute Type not valid. Got {Type}.", AttributeTypeNameInternal(A.Type));
 		return FAnimDatabaseFrameAttribute();
 	}
 
@@ -6312,6 +6806,11 @@ FAnimDatabaseFrameAttribute UAnimDatabaseFrameAttributeLibrary::FrameAttributeAd
 	{
 		UE::Learning::FrameAttribute::AddConstant(*Out.FrameAttribute, *A.FrameAttribute, { B });
 	}
+	else if (A.Type == EAnimDatabaseAttributeType::Angle)
+	{
+		UE::Learning::FrameAttribute::AddConstant(*Out.FrameAttribute, *A.FrameAttribute, { B });
+		UE::Learning::FrameAttribute::WrapAngleInplace(*Out.FrameAttribute);
+	}
 	else
 	{
 		UE::Learning::FrameAttribute::AddConstant(*Out.FrameAttribute, *A.FrameAttribute, { B, B, B });
@@ -6332,9 +6831,10 @@ FAnimDatabaseFrameAttribute UAnimDatabaseFrameAttributeLibrary::FrameAttributeSu
 		A.Type != EAnimDatabaseAttributeType::Scale &&
 		A.Type != EAnimDatabaseAttributeType::LinearVelocity &&
 		A.Type != EAnimDatabaseAttributeType::AngularVelocity &&
-		A.Type != EAnimDatabaseAttributeType::ScalarVelocity)
+		A.Type != EAnimDatabaseAttributeType::ScalarVelocity &&
+		A.Type != EAnimDatabaseAttributeType::Angle)
 	{
-		UE_LOGFMT(LogAnimDatabase, Error, "FrameAttributeSubtractFloat: FrameAttribute Type not valid. Got {Type}.", *UEnum::GetValueAsString(A.Type));
+		UE_LOGFMT(LogAnimDatabase, Error, "FrameAttributeSubtractFloat: FrameAttribute Type not valid. Got {Type}.", AttributeTypeNameInternal(A.Type));
 		return FAnimDatabaseFrameAttribute();
 	}
 
@@ -6344,6 +6844,11 @@ FAnimDatabaseFrameAttribute UAnimDatabaseFrameAttributeLibrary::FrameAttributeSu
 	{
 		UE::Learning::FrameAttribute::SubConstant(*Out.FrameAttribute, *A.FrameAttribute, { B });
 	}
+	else if (A.Type == EAnimDatabaseAttributeType::Angle)
+	{
+		UE::Learning::FrameAttribute::SubConstant(*Out.FrameAttribute, *A.FrameAttribute, { B });
+		UE::Learning::FrameAttribute::WrapAngleInplace(*Out.FrameAttribute);
+	}
 	else
 	{
 		UE::Learning::FrameAttribute::SubConstant(*Out.FrameAttribute, *A.FrameAttribute, { B, B, B });
@@ -6364,9 +6869,10 @@ FAnimDatabaseFrameAttribute UAnimDatabaseFrameAttributeLibrary::FrameAttributeMu
 		A.Type != EAnimDatabaseAttributeType::Scale &&
 		A.Type != EAnimDatabaseAttributeType::LinearVelocity &&
 		A.Type != EAnimDatabaseAttributeType::AngularVelocity &&
-		A.Type != EAnimDatabaseAttributeType::ScalarVelocity)
+		A.Type != EAnimDatabaseAttributeType::ScalarVelocity &&
+		A.Type != EAnimDatabaseAttributeType::Angle)
 	{
-		UE_LOGFMT(LogAnimDatabase, Error, "FrameAttributeMultiplyFloat: FrameAttribute Type not valid. Got {Type}.", *UEnum::GetValueAsString(A.Type));
+		UE_LOGFMT(LogAnimDatabase, Error, "FrameAttributeMultiplyFloat: FrameAttribute Type not valid. Got {Type}.", AttributeTypeNameInternal(A.Type));
 		return FAnimDatabaseFrameAttribute();
 	}
 
@@ -6376,6 +6882,11 @@ FAnimDatabaseFrameAttribute UAnimDatabaseFrameAttributeLibrary::FrameAttributeMu
 	{
 		UE::Learning::FrameAttribute::MulConstant(*Out.FrameAttribute, *A.FrameAttribute, { B });
 	}
+	else if (A.Type == EAnimDatabaseAttributeType::Angle)
+	{
+		UE::Learning::FrameAttribute::MulConstant(*Out.FrameAttribute, *A.FrameAttribute, { B });
+		UE::Learning::FrameAttribute::WrapAngleInplace(*Out.FrameAttribute);
+	}
 	else
 	{
 		UE::Learning::FrameAttribute::MulConstant(*Out.FrameAttribute, *A.FrameAttribute, { B, B, B });
@@ -6396,9 +6907,10 @@ FAnimDatabaseFrameAttribute UAnimDatabaseFrameAttributeLibrary::FrameAttributeDi
 		A.Type != EAnimDatabaseAttributeType::Scale &&
 		A.Type != EAnimDatabaseAttributeType::LinearVelocity &&
 		A.Type != EAnimDatabaseAttributeType::AngularVelocity &&
-		A.Type != EAnimDatabaseAttributeType::ScalarVelocity)
+		A.Type != EAnimDatabaseAttributeType::ScalarVelocity &&
+		A.Type != EAnimDatabaseAttributeType::Angle)
 	{
-		UE_LOGFMT(LogAnimDatabase, Error, "FrameAttributeDivideFloat: FrameAttribute Type not valid. Got {Type}.", *UEnum::GetValueAsString(A.Type));
+		UE_LOGFMT(LogAnimDatabase, Error, "FrameAttributeDivideFloat: FrameAttribute Type not valid. Got {Type}.", AttributeTypeNameInternal(A.Type));
 		return FAnimDatabaseFrameAttribute();
 	}
 
@@ -6408,6 +6920,11 @@ FAnimDatabaseFrameAttribute UAnimDatabaseFrameAttributeLibrary::FrameAttributeDi
 	{
 		UE::Learning::FrameAttribute::DivConstant(*Out.FrameAttribute, *A.FrameAttribute, { B });
 	}
+	else if (A.Type == EAnimDatabaseAttributeType::Angle)
+	{
+		UE::Learning::FrameAttribute::DivConstant(*Out.FrameAttribute, *A.FrameAttribute, { B });
+		UE::Learning::FrameAttribute::WrapAngleInplace(*Out.FrameAttribute);
+	}
 	else
 	{
 		UE::Learning::FrameAttribute::DivConstant(*Out.FrameAttribute, *A.FrameAttribute, { B, B, B });
@@ -6429,7 +6946,7 @@ FAnimDatabaseFrameAttribute UAnimDatabaseFrameAttributeLibrary::FrameAttributeAd
 		A.Type != EAnimDatabaseAttributeType::AngularVelocity &&
 		A.Type != EAnimDatabaseAttributeType::ScalarVelocity)
 	{
-		UE_LOGFMT(LogAnimDatabase, Error, "FrameAttributeAddVector: FrameAttribute Type not valid. Got {Type}.", *UEnum::GetValueAsString(A.Type));
+		UE_LOGFMT(LogAnimDatabase, Error, "FrameAttributeAddVector: FrameAttribute Type not valid. Got {Type}.", AttributeTypeNameInternal(A.Type));
 		return FAnimDatabaseFrameAttribute();
 	}
 
@@ -6455,7 +6972,7 @@ FAnimDatabaseFrameAttribute UAnimDatabaseFrameAttributeLibrary::FrameAttributeSu
 		A.Type != EAnimDatabaseAttributeType::AngularVelocity &&
 		A.Type != EAnimDatabaseAttributeType::ScalarVelocity)
 	{
-		UE_LOGFMT(LogAnimDatabase, Error, "FrameAttributeSubtractVector: FrameAttribute Type not valid. Got {Type}.", *UEnum::GetValueAsString(A.Type));
+		UE_LOGFMT(LogAnimDatabase, Error, "FrameAttributeSubtractVector: FrameAttribute Type not valid. Got {Type}.", AttributeTypeNameInternal(A.Type));
 		return FAnimDatabaseFrameAttribute();
 	}
 
@@ -6479,7 +6996,7 @@ FAnimDatabaseFrameAttribute UAnimDatabaseFrameAttributeLibrary::FrameAttributeMu
 		A.Type != EAnimDatabaseAttributeType::AngularVelocity &&
 		A.Type != EAnimDatabaseAttributeType::ScalarVelocity)
 	{
-		UE_LOGFMT(LogAnimDatabase, Error, "FrameAttributeMultiplyVector: FrameAttribute Type not valid. Got {Type}.", *UEnum::GetValueAsString(A.Type));
+		UE_LOGFMT(LogAnimDatabase, Error, "FrameAttributeMultiplyVector: FrameAttribute Type not valid. Got {Type}.", AttributeTypeNameInternal(A.Type));
 		return FAnimDatabaseFrameAttribute();
 	}
 
@@ -6503,7 +7020,7 @@ FAnimDatabaseFrameAttribute UAnimDatabaseFrameAttributeLibrary::FrameAttributeDi
 		A.Type != EAnimDatabaseAttributeType::AngularVelocity &&
 		A.Type != EAnimDatabaseAttributeType::ScalarVelocity)
 	{
-		UE_LOGFMT(LogAnimDatabase, Error, "FrameAttributeDivideVector: FrameAttribute Type not valid. Got {Type}.", *UEnum::GetValueAsString(A.Type));
+		UE_LOGFMT(LogAnimDatabase, Error, "FrameAttributeDivideVector: FrameAttribute Type not valid. Got {Type}.", AttributeTypeNameInternal(A.Type));
 		return FAnimDatabaseFrameAttribute();
 	}
 
@@ -6525,7 +7042,7 @@ FAnimDatabaseFrameAttribute UAnimDatabaseFrameAttributeLibrary::FrameAttributeQu
 
 	if (B.Type != EAnimDatabaseAttributeType::Rotation)
 	{
-		UE_LOGFMT(LogAnimDatabase, Error, "FrameAttributeQuatMultiply: FrameAttribute Type not valid. Got {Type}. Expected Rotation.", *UEnum::GetValueAsString(B.Type));
+		UE_LOGFMT(LogAnimDatabase, Error, "FrameAttributeQuatMultiply: FrameAttribute Type not valid. Got {Type}. Expected Rotation.", AttributeTypeNameInternal(B.Type));
 		return FAnimDatabaseFrameAttribute();
 	}
 
@@ -6547,7 +7064,7 @@ FAnimDatabaseFrameAttribute UAnimDatabaseFrameAttributeLibrary::FrameAttributeQu
 
 	if (B.Type != EAnimDatabaseAttributeType::Rotation)
 	{
-		UE_LOGFMT(LogAnimDatabase, Error, "FrameAttributeQuatDivide: FrameAttribute Type not valid. Got {Type}. Expected Rotation.", *UEnum::GetValueAsString(B.Type));
+		UE_LOGFMT(LogAnimDatabase, Error, "FrameAttributeQuatDivide: FrameAttribute Type not valid. Got {Type}. Expected Rotation.", AttributeTypeNameInternal(B.Type));
 		return FAnimDatabaseFrameAttribute();
 	}
 
@@ -6569,7 +7086,7 @@ FAnimDatabaseFrameAttribute UAnimDatabaseFrameAttributeLibrary::FrameAttributeMu
 
 	if (A.Type != EAnimDatabaseAttributeType::Rotation)
 	{
-		UE_LOGFMT(LogAnimDatabase, Error, "FrameAttributeMultiplyQuat: FrameAttribute Type not valid. Got {Type}. Expected Rotation.", *UEnum::GetValueAsString(A.Type));
+		UE_LOGFMT(LogAnimDatabase, Error, "FrameAttributeMultiplyQuat: FrameAttribute Type not valid. Got {Type}. Expected Rotation.", AttributeTypeNameInternal(A.Type));
 		return FAnimDatabaseFrameAttribute();
 	}
 
@@ -6591,7 +7108,7 @@ FAnimDatabaseFrameAttribute UAnimDatabaseFrameAttributeLibrary::FrameAttributeDi
 
 	if (A.Type != EAnimDatabaseAttributeType::Rotation)
 	{
-		UE_LOGFMT(LogAnimDatabase, Error, "FrameAttributeDivideQuat: FrameAttribute Type not valid. Got {Type}. Expected Rotation.", *UEnum::GetValueAsString(A.Type));
+		UE_LOGFMT(LogAnimDatabase, Error, "FrameAttributeDivideQuat: FrameAttribute Type not valid. Got {Type}. Expected Rotation.", AttributeTypeNameInternal(A.Type));
 		return FAnimDatabaseFrameAttribute();
 	}
 
@@ -6613,7 +7130,7 @@ FAnimDatabaseFrameAttribute UAnimDatabaseFrameAttributeLibrary::FrameAttributeTr
 
 	if (B.Type != EAnimDatabaseAttributeType::Transform)
 	{
-		UE_LOGFMT(LogAnimDatabase, Error, "FrameAttributeTransformMultiply: FrameAttribute Type not valid. Got {Type}. Expected Transform.", *UEnum::GetValueAsString(B.Type));
+		UE_LOGFMT(LogAnimDatabase, Error, "FrameAttributeTransformMultiply: FrameAttribute Type not valid. Got {Type}. Expected Transform.", AttributeTypeNameInternal(B.Type));
 		return FAnimDatabaseFrameAttribute();
 	}
 
@@ -6635,7 +7152,7 @@ FAnimDatabaseFrameAttribute UAnimDatabaseFrameAttributeLibrary::FrameAttributeTr
 
 	if (B.Type != EAnimDatabaseAttributeType::Transform)
 	{
-		UE_LOGFMT(LogAnimDatabase, Error, "FrameAttributeTransformDivide: FrameAttribute Type not valid. Got {Type}. Expected Transform.", *UEnum::GetValueAsString(B.Type));
+		UE_LOGFMT(LogAnimDatabase, Error, "FrameAttributeTransformDivide: FrameAttribute Type not valid. Got {Type}. Expected Transform.", AttributeTypeNameInternal(B.Type));
 		return FAnimDatabaseFrameAttribute();
 	}
 
@@ -6657,7 +7174,7 @@ FAnimDatabaseFrameAttribute UAnimDatabaseFrameAttributeLibrary::FrameAttributeMu
 
 	if (A.Type != EAnimDatabaseAttributeType::Transform)
 	{
-		UE_LOGFMT(LogAnimDatabase, Error, "FrameAttributeMultiplyTransform: FrameAttribute Type not valid. Got {Type}. Expected Transform.", *UEnum::GetValueAsString(A.Type));
+		UE_LOGFMT(LogAnimDatabase, Error, "FrameAttributeMultiplyTransform: FrameAttribute Type not valid. Got {Type}. Expected Transform.", AttributeTypeNameInternal(A.Type));
 		return FAnimDatabaseFrameAttribute();
 	}
 
@@ -6679,7 +7196,7 @@ FAnimDatabaseFrameAttribute UAnimDatabaseFrameAttributeLibrary::FrameAttributeDi
 
 	if (A.Type != EAnimDatabaseAttributeType::Transform)
 	{
-		UE_LOGFMT(LogAnimDatabase, Error, "FrameAttributeDivideTransform: FrameAttribute Type not valid. Got {Type}. Expected Transform.", *UEnum::GetValueAsString(A.Type));
+		UE_LOGFMT(LogAnimDatabase, Error, "FrameAttributeDivideTransform: FrameAttribute Type not valid. Got {Type}. Expected Transform.", AttributeTypeNameInternal(A.Type));
 		return FAnimDatabaseFrameAttribute();
 	}
 
@@ -6704,7 +7221,7 @@ FAnimDatabaseFrameAttribute UAnimDatabaseFrameAttributeLibrary::FrameAttributeTr
 		B.Type != EAnimDatabaseAttributeType::LinearVelocity &&
 		B.Type != EAnimDatabaseAttributeType::AngularVelocity)
 	{
-		UE_LOGFMT(LogAnimDatabase, Error, "FrameAttributeTransformApply: FrameAttribute Type not valid. Got {Type}. Expected Location, Direction, LinearVelocity or AngularVelocity.", *UEnum::GetValueAsString(B.Type));
+		UE_LOGFMT(LogAnimDatabase, Error, "FrameAttributeTransformApply: FrameAttribute Type not valid. Got {Type}. Expected Location, Direction, LinearVelocity or AngularVelocity.", AttributeTypeNameInternal(B.Type));
 		return FAnimDatabaseFrameAttribute();
 	}
 
@@ -6725,37 +7242,129 @@ FAnimDatabaseFrameAttribute UAnimDatabaseFrameAttributeLibrary::FrameAttributeTr
 	return Out;
 }
 
-void UAnimDatabaseFrameAttributeLibrary::FrameAttributesAdd(TArray<FAnimDatabaseFrameAttribute>& Out, const TArray<FAnimDatabaseFrameAttribute>& A, const FAnimDatabaseFrameAttribute& B)
+FAnimDatabaseFrameAttribute UAnimDatabaseFrameAttributeLibrary::FrameAttributeAddAngleDegrees(const FAnimDatabaseFrameAttribute& A, const float B)
 {
-	Out.SetNum(A.Num());
-	FrameAttributesAddArrayView(Out, A, B);
-}
+	TRACE_CPUPROFILER_EVENT_SCOPE(UAnimDatabaseFrameAttributeLibrary::FrameAttributeAddAngleDegrees);
 
-void UAnimDatabaseFrameAttributeLibrary::FrameAttributesAddArrayView(const TArrayView<FAnimDatabaseFrameAttribute> Out, const TArrayView<const FAnimDatabaseFrameAttribute> A, const FAnimDatabaseFrameAttribute& B)
-{
-	check(Out.Num() == A.Num());
-	ParallelFor(Out.Num(), [&](int32 Idx) {
-		Out[Idx] = FrameAttributeAdd(A[Idx], B);
-		});
-}
+	if (!A.IsValid())
+	{
+		UE_LOGFMT(LogAnimDatabase, Error, "FrameAttributeAddAngleDegrees: Invalid FrameAttribute.");
+		return FAnimDatabaseFrameAttribute();
+	}
 
-void UAnimDatabaseFrameAttributeLibrary::FrameAttributesSubtract(TArray<FAnimDatabaseFrameAttribute>& Out, const TArray<FAnimDatabaseFrameAttribute>& A, const FAnimDatabaseFrameAttribute& B)
-{
-	Out.SetNum(A.Num());
-	FrameAttributesSubtractArrayView(Out, A, B);
-}
+	if (A.Type != EAnimDatabaseAttributeType::Angle)
+	{
+		UE_LOGFMT(LogAnimDatabase, Error, "FrameAttributeAddAngleDegrees: FrameAttribute Type not valid. Got {Type}. Expected Angle.", AttributeTypeNameInternal(A.Type));
+		return FAnimDatabaseFrameAttribute();
+	}
 
-void UAnimDatabaseFrameAttributeLibrary::FrameAttributesSubtractArrayView(const TArrayView<FAnimDatabaseFrameAttribute> Out, const TArrayView<const FAnimDatabaseFrameAttribute> A, const FAnimDatabaseFrameAttribute& B)
-{
-	check(Out.Num() == A.Num());
-	ParallelFor(Out.Num(), [&](int32 Idx) {
-		Out[Idx] = FrameAttributeSubtract(A[Idx], B);
-		});
+	FAnimDatabaseFrameAttribute Out = MakeEmptyFrameAttribute();
+	Out.Type = A.Type;
+	UE::Learning::FrameAttribute::AddConstant(*Out.FrameAttribute, *A.FrameAttribute, { FMath::DegreesToRadians(B) });
+	UE::Learning::FrameAttribute::WrapAngleInplace(*Out.FrameAttribute);
+	return Out;
 }
 
-void UAnimDatabaseFrameAttributeLibrary::FrameAttributesMultiply(TArray<FAnimDatabaseFrameAttribute>& Out, const TArray<FAnimDatabaseFrameAttribute>& A, const FAnimDatabaseFrameAttribute& B)
+FAnimDatabaseFrameAttribute UAnimDatabaseFrameAttributeLibrary::FrameAttributeSubtractAngleDegrees(const FAnimDatabaseFrameAttribute& A, const float B)
 {
-	Out.SetNum(A.Num());
+	TRACE_CPUPROFILER_EVENT_SCOPE(UAnimDatabaseFrameAttributeLibrary::FrameAttributeSubtractAngleDegrees);
+
+	if (!A.IsValid())
+	{
+		UE_LOGFMT(LogAnimDatabase, Error, "FrameAttributeSubtractAngleDegrees: Invalid FrameAttribute.");
+		return FAnimDatabaseFrameAttribute();
+	}
+
+	if (A.Type != EAnimDatabaseAttributeType::Angle)
+	{
+		UE_LOGFMT(LogAnimDatabase, Error, "FrameAttributeSubtractAngleDegrees: FrameAttribute Type not valid. Got {Type}. Expected Angle.", AttributeTypeNameInternal(A.Type));
+		return FAnimDatabaseFrameAttribute();
+	}
+
+	FAnimDatabaseFrameAttribute Out = MakeEmptyFrameAttribute();
+	Out.Type = A.Type;
+	UE::Learning::FrameAttribute::SubConstant(*Out.FrameAttribute, *A.FrameAttribute, { FMath::DegreesToRadians(B) });
+	UE::Learning::FrameAttribute::WrapAngleInplace(*Out.FrameAttribute);
+	return Out;
+}
+
+FAnimDatabaseFrameAttribute UAnimDatabaseFrameAttributeLibrary::FrameAttributeAddAngleRadians(const FAnimDatabaseFrameAttribute& A, const float B)
+{
+	TRACE_CPUPROFILER_EVENT_SCOPE(UAnimDatabaseFrameAttributeLibrary::FrameAttributeAddAngleRadians);
+
+	if (!A.IsValid())
+	{
+		UE_LOGFMT(LogAnimDatabase, Error, "FrameAttributeAddAngleRadians: Invalid FrameAttribute.");
+		return FAnimDatabaseFrameAttribute();
+	}
+
+	if (A.Type != EAnimDatabaseAttributeType::Angle)
+	{
+		UE_LOGFMT(LogAnimDatabase, Error, "FrameAttributeAddAngleRadians: FrameAttribute Type not valid. Got {Type}. Expected Angle.", AttributeTypeNameInternal(A.Type));
+		return FAnimDatabaseFrameAttribute();
+	}
+
+	FAnimDatabaseFrameAttribute Out = MakeEmptyFrameAttribute();
+	Out.Type = A.Type;
+	UE::Learning::FrameAttribute::AddConstant(*Out.FrameAttribute, *A.FrameAttribute, { B });
+	UE::Learning::FrameAttribute::WrapAngleInplace(*Out.FrameAttribute);
+	return Out;
+}
+
+FAnimDatabaseFrameAttribute UAnimDatabaseFrameAttributeLibrary::FrameAttributeSubtractAngleRadians(const FAnimDatabaseFrameAttribute& A, const float B)
+{
+	TRACE_CPUPROFILER_EVENT_SCOPE(UAnimDatabaseFrameAttributeLibrary::FrameAttributeSubtractAngleRadians);
+
+	if (!A.IsValid())
+	{
+		UE_LOGFMT(LogAnimDatabase, Error, "FrameAttributeSubtractAngleRadians: Invalid FrameAttribute.");
+		return FAnimDatabaseFrameAttribute();
+	}
+
+	if (A.Type != EAnimDatabaseAttributeType::Angle)
+	{
+		UE_LOGFMT(LogAnimDatabase, Error, "FrameAttributeSubtractAngleRadians: FrameAttribute Type not valid. Got {Type}. Expected Angle.", AttributeTypeNameInternal(A.Type));
+		return FAnimDatabaseFrameAttribute();
+	}
+
+	FAnimDatabaseFrameAttribute Out = MakeEmptyFrameAttribute();
+	Out.Type = A.Type;
+	UE::Learning::FrameAttribute::SubConstant(*Out.FrameAttribute, *A.FrameAttribute, { B });
+	UE::Learning::FrameAttribute::WrapAngleInplace(*Out.FrameAttribute);
+	return Out;
+}
+
+void UAnimDatabaseFrameAttributeLibrary::FrameAttributesAdd(TArray<FAnimDatabaseFrameAttribute>& Out, const TArray<FAnimDatabaseFrameAttribute>& A, const FAnimDatabaseFrameAttribute& B)
+{
+	Out.SetNum(A.Num());
+	FrameAttributesAddArrayView(Out, A, B);
+}
+
+void UAnimDatabaseFrameAttributeLibrary::FrameAttributesAddArrayView(const TArrayView<FAnimDatabaseFrameAttribute> Out, const TArrayView<const FAnimDatabaseFrameAttribute> A, const FAnimDatabaseFrameAttribute& B)
+{
+	check(Out.Num() == A.Num());
+	ParallelFor(Out.Num(), [&](int32 Idx) {
+		Out[Idx] = FrameAttributeAdd(A[Idx], B);
+		});
+}
+
+void UAnimDatabaseFrameAttributeLibrary::FrameAttributesSubtract(TArray<FAnimDatabaseFrameAttribute>& Out, const TArray<FAnimDatabaseFrameAttribute>& A, const FAnimDatabaseFrameAttribute& B)
+{
+	Out.SetNum(A.Num());
+	FrameAttributesSubtractArrayView(Out, A, B);
+}
+
+void UAnimDatabaseFrameAttributeLibrary::FrameAttributesSubtractArrayView(const TArrayView<FAnimDatabaseFrameAttribute> Out, const TArrayView<const FAnimDatabaseFrameAttribute> A, const FAnimDatabaseFrameAttribute& B)
+{
+	check(Out.Num() == A.Num());
+	ParallelFor(Out.Num(), [&](int32 Idx) {
+		Out[Idx] = FrameAttributeSubtract(A[Idx], B);
+		});
+}
+
+void UAnimDatabaseFrameAttributeLibrary::FrameAttributesMultiply(TArray<FAnimDatabaseFrameAttribute>& Out, const TArray<FAnimDatabaseFrameAttribute>& A, const FAnimDatabaseFrameAttribute& B)
+{
+	Out.SetNum(A.Num());
 	FrameAttributesMultiplyArrayView(Out, A, B);
 }
 
@@ -6805,7 +7414,7 @@ FAnimDatabaseFrameAttribute UAnimDatabaseFrameAttributeLibrary::FrameAttributeRo
 		(Vector.Type != EAnimDatabaseAttributeType::LinearVelocity) &&
 		(Vector.Type != EAnimDatabaseAttributeType::AngularVelocity)))
 	{
-		UE_LOGFMT(LogAnimDatabase, Error, "FrameAttributeRotate: FrameAttribute Types don't match. Got {Type0} and {Type1}. Expected Rotation and Location, Direction, or Velocity.", *UEnum::GetValueAsString(Rotation.Type), *UEnum::GetValueAsString(Vector.Type));
+		UE_LOGFMT(LogAnimDatabase, Error, "FrameAttributeRotate: FrameAttribute Types don't match. Got {Type0} and {Type1}. Expected Rotation and Location, Direction, or Velocity.", AttributeTypeNameInternal(Rotation.Type), AttributeTypeNameInternal(Vector.Type));
 		return FAnimDatabaseFrameAttribute();
 	}
 
@@ -6829,7 +7438,7 @@ FAnimDatabaseFrameAttribute UAnimDatabaseFrameAttributeLibrary::FrameAttributeUn
 		(Vector.Type != EAnimDatabaseAttributeType::LinearVelocity) &&
 		(Vector.Type != EAnimDatabaseAttributeType::AngularVelocity)))
 	{
-		UE_LOGFMT(LogAnimDatabase, Error, "FrameAttributeUnrotate: FrameAttribute Types don't match. Got {Type0} and {Type1}. Expected Rotation and Location, Direction, or Velocity.", *UEnum::GetValueAsString(Rotation.Type), *UEnum::GetValueAsString(Vector.Type));
+		UE_LOGFMT(LogAnimDatabase, Error, "FrameAttributeUnrotate: FrameAttribute Types don't match. Got {Type0} and {Type1}. Expected Rotation and Location, Direction, or Velocity.", AttributeTypeNameInternal(Rotation.Type), AttributeTypeNameInternal(Vector.Type));
 		return FAnimDatabaseFrameAttribute();
 	}
 
@@ -6852,7 +7461,7 @@ FAnimDatabaseFrameAttribute UAnimDatabaseFrameAttributeLibrary::FrameAttributeQu
 		Vector.Type != EAnimDatabaseAttributeType::LinearVelocity &&
 		Vector.Type != EAnimDatabaseAttributeType::AngularVelocity)
 	{
-		UE_LOGFMT(LogAnimDatabase, Error, "FrameAttributeQuatRotate: FrameAttribute Types don't match. Got {Type1}. Expected Location, Direction, or Velocity.", *UEnum::GetValueAsString(Vector.Type));
+		UE_LOGFMT(LogAnimDatabase, Error, "FrameAttributeQuatRotate: FrameAttribute Types don't match. Got {Type1}. Expected Location, Direction, or Velocity.", AttributeTypeNameInternal(Vector.Type));
 		return FAnimDatabaseFrameAttribute();
 	}
 
@@ -6875,7 +7484,7 @@ FAnimDatabaseFrameAttribute UAnimDatabaseFrameAttributeLibrary::FrameAttributeQu
 		Vector.Type != EAnimDatabaseAttributeType::LinearVelocity &&
 		Vector.Type != EAnimDatabaseAttributeType::AngularVelocity)
 	{
-		UE_LOGFMT(LogAnimDatabase, Error, "FrameAttributeQuatUnrotate: FrameAttribute Types don't match. Got {Type1}. Expected Location, Direction, or Velocity.", *UEnum::GetValueAsString(Vector.Type));
+		UE_LOGFMT(LogAnimDatabase, Error, "FrameAttributeQuatUnrotate: FrameAttribute Types don't match. Got {Type1}. Expected Location, Direction, or Velocity.", AttributeTypeNameInternal(Vector.Type));
 		return FAnimDatabaseFrameAttribute();
 	}
 
@@ -6977,7 +7586,7 @@ FAnimDatabaseFrameAttribute UAnimDatabaseFrameAttributeLibrary::FrameAttributeRo
 
 	if (Rotation.Type != EAnimDatabaseAttributeType::Rotation)
 	{
-		UE_LOGFMT(LogAnimDatabase, Error, "FrameAttributeRotateDirection: FrameAttribute Types don't match. Got {Type}. Expected Rotation.", *UEnum::GetValueAsString(Rotation.Type));
+		UE_LOGFMT(LogAnimDatabase, Error, "FrameAttributeRotateDirection: FrameAttribute Types don't match. Got {Type}. Expected Rotation.", AttributeTypeNameInternal(Rotation.Type));
 		return FAnimDatabaseFrameAttribute();
 	}
 
@@ -6997,7 +7606,7 @@ FAnimDatabaseFrameAttribute UAnimDatabaseFrameAttributeLibrary::FrameAttributeUn
 
 	if (Rotation.Type != EAnimDatabaseAttributeType::Rotation)
 	{
-		UE_LOGFMT(LogAnimDatabase, Error, "FrameAttributeUnrotateDirection: FrameAttribute Types don't match. Got {Type}. Expected Rotation.", *UEnum::GetValueAsString(Rotation.Type));
+		UE_LOGFMT(LogAnimDatabase, Error, "FrameAttributeUnrotateDirection: FrameAttribute Types don't match. Got {Type}. Expected Rotation.", AttributeTypeNameInternal(Rotation.Type));
 		return FAnimDatabaseFrameAttribute();
 	}
 
@@ -7017,7 +7626,7 @@ FAnimDatabaseFrameAttribute UAnimDatabaseFrameAttributeLibrary::FrameAttributeRo
 
 	if (Directions.Type != EAnimDatabaseAttributeType::Direction)
 	{
-		UE_LOGFMT(LogAnimDatabase, Error, "FrameAttributeRotationBetweenDirection: FrameAttribute Types don't match. Got {Type}. Expected Direction.", *UEnum::GetValueAsString(Directions.Type));
+		UE_LOGFMT(LogAnimDatabase, Error, "FrameAttributeRotationBetweenDirection: FrameAttribute Types don't match. Got {Type}. Expected Direction.", AttributeTypeNameInternal(Directions.Type));
 		return FAnimDatabaseFrameAttribute();
 	}
 
@@ -7041,7 +7650,7 @@ FAnimDatabaseFrameAttribute UAnimDatabaseFrameAttributeLibrary::FrameAttributePr
 		A.Type != EAnimDatabaseAttributeType::AngularVelocity &&
 		A.Type != EAnimDatabaseAttributeType::ScalarVelocity)
 	{
-		UE_LOGFMT(LogAnimDatabase, Error, "FrameAttributeProject: FrameAttribute Types don't match. Got {Type}. Expected Location, Direction or Velocity.", *UEnum::GetValueAsString(A.Type));
+		UE_LOGFMT(LogAnimDatabase, Error, "FrameAttributeProject: FrameAttribute Types don't match. Got {Type}. Expected Location, Direction or Velocity.", AttributeTypeNameInternal(A.Type));
 		return FAnimDatabaseFrameAttribute();
 	}
 
@@ -7061,6 +7670,20 @@ FAnimDatabaseFrameAttribute UAnimDatabaseFrameAttributeLibrary::FrameAttributePr
 	return Out;
 }
 
+FAnimDatabaseFrameAttribute UAnimDatabaseFrameAttributeLibrary::FrameAttributeCopy(const FAnimDatabaseFrameAttribute& A)
+{
+	if (!A.IsValid())
+	{
+		UE_LOGFMT(LogAnimDatabase, Error, "FrameAttributeCopy: Invalid FrameAttribute.");
+		return FAnimDatabaseFrameAttribute();
+	}
+
+	FAnimDatabaseFrameAttribute Out = MakeEmptyFrameAttribute();
+	Out.Type = A.Type;
+	UE::Learning::FrameAttribute::Copy(*Out.FrameAttribute, *A.FrameAttribute);
+	return Out;
+}
+
 FAnimDatabaseFrameAttribute UAnimDatabaseFrameAttributeLibrary::FrameAttributeNegate(const FAnimDatabaseFrameAttribute& A)
 {
 	if (!A.IsValid())
@@ -7075,9 +7698,10 @@ FAnimDatabaseFrameAttribute UAnimDatabaseFrameAttributeLibrary::FrameAttributeNe
 		A.Type != EAnimDatabaseAttributeType::LinearVelocity &&
 		A.Type != EAnimDatabaseAttributeType::AngularVelocity &&
 		A.Type != EAnimDatabaseAttributeType::ScalarVelocity &&
-		A.Type != EAnimDatabaseAttributeType::Direction)
+		A.Type != EAnimDatabaseAttributeType::Direction &&
+		A.Type != EAnimDatabaseAttributeType::Angle)
 	{
-		UE_LOGFMT(LogAnimDatabase, Error, "FrameAttributeNegate: FrameAttribute Types can't be negated. Got {Type}.", *UEnum::GetValueAsString(A.Type));
+		UE_LOGFMT(LogAnimDatabase, Error, "FrameAttributeNegate: FrameAttribute Types can't be negated. Got {Type}.", AttributeTypeNameInternal(A.Type));
 		return FAnimDatabaseFrameAttribute();
 	}
 
@@ -7097,7 +7721,7 @@ FAnimDatabaseFrameAttribute UAnimDatabaseFrameAttributeLibrary::FrameAttributeIn
 
 	if (A.Type != EAnimDatabaseAttributeType::Rotation && A.Type != EAnimDatabaseAttributeType::Transform)
 	{
-		UE_LOGFMT(LogAnimDatabase, Error, "FrameAttributeInvert: FrameAttribute Types don't match. Got {Type}. Expected Rotation or Transform.", *UEnum::GetValueAsString(A.Type));
+		UE_LOGFMT(LogAnimDatabase, Error, "FrameAttributeInvert: FrameAttribute Types don't match. Got {Type}. Expected Rotation or Transform.", AttributeTypeNameInternal(A.Type));
 		return FAnimDatabaseFrameAttribute();
 	}
 
@@ -7138,9 +7762,10 @@ FAnimDatabaseFrameAttribute UAnimDatabaseFrameAttributeLibrary::FrameAttributeAb
 		A.Type != EAnimDatabaseAttributeType::Scale &&
 		A.Type != EAnimDatabaseAttributeType::LinearVelocity &&
 		A.Type != EAnimDatabaseAttributeType::AngularVelocity &&
-		A.Type != EAnimDatabaseAttributeType::ScalarVelocity)
+		A.Type != EAnimDatabaseAttributeType::ScalarVelocity &&
+		A.Type != EAnimDatabaseAttributeType::Angle)
 	{
-		UE_LOGFMT(LogAnimDatabase, Error, "FrameAttributeAbs: FrameAttribute Types invalid. Got {Type}.", *UEnum::GetValueAsString(A.Type));
+		UE_LOGFMT(LogAnimDatabase, Error, "FrameAttributeAbs: FrameAttribute Types invalid. Got {Type}.", AttributeTypeNameInternal(A.Type));
 		return FAnimDatabaseFrameAttribute();
 	}
 
@@ -7175,7 +7800,7 @@ FAnimDatabaseFrameAttribute UAnimDatabaseFrameAttributeLibrary::FrameAttributeLo
 		A.Type != EAnimDatabaseAttributeType::AngularVelocity &&
 		A.Type != EAnimDatabaseAttributeType::ScalarVelocity)
 	{
-		UE_LOGFMT(LogAnimDatabase, Error, "FrameAttributeLog: FrameAttribute Types invalid. Got {Type}.", *UEnum::GetValueAsString(A.Type));
+		UE_LOGFMT(LogAnimDatabase, Error, "FrameAttributeLog: FrameAttribute Types invalid. Got {Type}.", AttributeTypeNameInternal(A.Type));
 		return FAnimDatabaseFrameAttribute();
 	}
 
@@ -7200,7 +7825,7 @@ FAnimDatabaseFrameAttribute UAnimDatabaseFrameAttributeLibrary::FrameAttributeEx
 		A.Type != EAnimDatabaseAttributeType::AngularVelocity &&
 		A.Type != EAnimDatabaseAttributeType::ScalarVelocity)
 	{
-		UE_LOGFMT(LogAnimDatabase, Error, "FrameAttributeExp: FrameAttribute Types invalid. Got {Type}.", *UEnum::GetValueAsString(A.Type));
+		UE_LOGFMT(LogAnimDatabase, Error, "FrameAttributeExp: FrameAttribute Types invalid. Got {Type}.", AttributeTypeNameInternal(A.Type));
 		return FAnimDatabaseFrameAttribute();
 	}
 
@@ -7226,7 +7851,7 @@ FAnimDatabaseFrameAttribute UAnimDatabaseFrameAttributeLibrary::FrameAttributeLe
 		Vector.Type != EAnimDatabaseAttributeType::AngularVelocity &&
 		Vector.Type != EAnimDatabaseAttributeType::ScalarVelocity)
 	{
-		UE_LOGFMT(LogAnimDatabase, Error, "FrameAttributeLength: FrameAttribute Types don't match. Got {Type}. Expected Location, Scale, LinearVelocity, AngularVelocity, or ScalarVelocity.", *UEnum::GetValueAsString(Vector.Type));
+		UE_LOGFMT(LogAnimDatabase, Error, "FrameAttributeLength: FrameAttribute Types don't match. Got {Type}. Expected Location, Scale, LinearVelocity, AngularVelocity, or ScalarVelocity.", AttributeTypeNameInternal(Vector.Type));
 		return FAnimDatabaseFrameAttribute();
 	}
 
@@ -7248,7 +7873,7 @@ FAnimDatabaseFrameAttribute UAnimDatabaseFrameAttributeLibrary::FrameAttributeTo
 
 	if (A.Type != EAnimDatabaseAttributeType::Rotation)
 	{
-		UE_LOGFMT(LogAnimDatabase, Error, "FrameAttributeToRotationVector: FrameAttribute Types don't match. Got {Type}. Expected Rotation.", *UEnum::GetValueAsString(A.Type));
+		UE_LOGFMT(LogAnimDatabase, Error, "FrameAttributeToRotationVector: FrameAttribute Types don't match. Got {Type}. Expected Rotation.", AttributeTypeNameInternal(A.Type));
 		return FAnimDatabaseFrameAttribute();
 	}
 
@@ -7268,7 +7893,7 @@ FAnimDatabaseFrameAttribute UAnimDatabaseFrameAttributeLibrary::FrameAttributeFr
 
 	if (A.Type != EAnimDatabaseAttributeType::AngularVelocity)
 	{
-		UE_LOGFMT(LogAnimDatabase, Error, "FrameAttributeToRotationVector: FrameAttribute Types don't match. Got {Type}. Expected Angular Velocity.", *UEnum::GetValueAsString(A.Type));
+		UE_LOGFMT(LogAnimDatabase, Error, "FrameAttributeToRotationVector: FrameAttribute Types don't match. Got {Type}. Expected Angular Velocity.", AttributeTypeNameInternal(A.Type));
 		return FAnimDatabaseFrameAttribute();
 	}
 
@@ -7293,7 +7918,7 @@ FAnimDatabaseFrameAttribute UAnimDatabaseFrameAttributeLibrary::FrameAttributeX(
 		Vector.Type != EAnimDatabaseAttributeType::AngularVelocity &&
 		Vector.Type != EAnimDatabaseAttributeType::ScalarVelocity)
 	{
-		UE_LOGFMT(LogAnimDatabase, Error, "FrameAttributeX: FrameAttribute Types don't match. Got {Type}. Expected Location, Rotation, Scale, LinearVelocity, AngularVelocity, or ScalarVelocity.", *UEnum::GetValueAsString(Vector.Type));
+		UE_LOGFMT(LogAnimDatabase, Error, "FrameAttributeX: FrameAttribute Types don't match. Got {Type}. Expected Location, Rotation, Scale, LinearVelocity, AngularVelocity, or ScalarVelocity.", AttributeTypeNameInternal(Vector.Type));
 		return FAnimDatabaseFrameAttribute();
 	}
 
@@ -7318,7 +7943,7 @@ FAnimDatabaseFrameAttribute UAnimDatabaseFrameAttributeLibrary::FrameAttributeY(
 		Vector.Type != EAnimDatabaseAttributeType::AngularVelocity &&
 		Vector.Type != EAnimDatabaseAttributeType::ScalarVelocity)
 	{
-		UE_LOGFMT(LogAnimDatabase, Error, "FrameAttributeY: FrameAttribute Types don't match. Got {Type}. Expected Location, Rotation, Scale, LinearVelocity, AngularVelocity, or ScalarVelocity.", *UEnum::GetValueAsString(Vector.Type));
+		UE_LOGFMT(LogAnimDatabase, Error, "FrameAttributeY: FrameAttribute Types don't match. Got {Type}. Expected Location, Rotation, Scale, LinearVelocity, AngularVelocity, or ScalarVelocity.", AttributeTypeNameInternal(Vector.Type));
 		return FAnimDatabaseFrameAttribute();
 	}
 
@@ -7343,7 +7968,7 @@ FAnimDatabaseFrameAttribute UAnimDatabaseFrameAttributeLibrary::FrameAttributeZ(
 		Vector.Type != EAnimDatabaseAttributeType::AngularVelocity &&
 		Vector.Type != EAnimDatabaseAttributeType::ScalarVelocity)
 	{
-		UE_LOGFMT(LogAnimDatabase, Error, "FrameAttributeZ: FrameAttribute Types don't match. Got {Type}. Expected Location, Rotation, Scale, LinearVelocity, AngularVelocity, or ScalarVelocity.", *UEnum::GetValueAsString(Vector.Type));
+		UE_LOGFMT(LogAnimDatabase, Error, "FrameAttributeZ: FrameAttribute Types don't match. Got {Type}. Expected Location, Rotation, Scale, LinearVelocity, AngularVelocity, or ScalarVelocity.", AttributeTypeNameInternal(Vector.Type));
 		return FAnimDatabaseFrameAttribute();
 	}
 
@@ -7363,7 +7988,7 @@ FAnimDatabaseFrameAttribute UAnimDatabaseFrameAttributeLibrary::FrameAttributeW(
 
 	if (Vector.Type != EAnimDatabaseAttributeType::Rotation)
 	{
-		UE_LOGFMT(LogAnimDatabase, Error, "FrameAttributeW: FrameAttribute Types don't match. Got {Type}. Expected Rotation.", *UEnum::GetValueAsString(Vector.Type));
+		UE_LOGFMT(LogAnimDatabase, Error, "FrameAttributeW: FrameAttribute Types don't match. Got {Type}. Expected Rotation.", AttributeTypeNameInternal(Vector.Type));
 		return FAnimDatabaseFrameAttribute();
 	}
 
@@ -7373,6 +7998,66 @@ FAnimDatabaseFrameAttribute UAnimDatabaseFrameAttributeLibrary::FrameAttributeW(
 	return Out;
 }
 
+FAnimDatabaseFrameAttribute UAnimDatabaseFrameAttributeLibrary::FrameAttributeYaw(const FAnimDatabaseFrameAttribute& A)
+{
+	if (!A.IsValid())
+	{
+		UE_LOGFMT(LogAnimDatabase, Error, "FrameAttributeYaw: Invalid FrameAttribute.");
+		return FAnimDatabaseFrameAttribute();
+	}
+
+	if (A.Type != EAnimDatabaseAttributeType::Rotation)
+	{
+		UE_LOGFMT(LogAnimDatabase, Error, "FrameAttributeYaw: FrameAttribute Types don't match. Got {Type}. Expected Rotation.", AttributeTypeNameInternal(A.Type));
+		return FAnimDatabaseFrameAttribute();
+	}
+
+	FAnimDatabaseFrameAttribute Out = MakeEmptyFrameAttribute();
+	Out.Type = EAnimDatabaseAttributeType::Angle;
+	UE::Learning::FrameAttribute::QuatYaw(*Out.FrameAttribute, *A.FrameAttribute);
+	return Out;
+}
+
+FAnimDatabaseFrameAttribute UAnimDatabaseFrameAttributeLibrary::FrameAttributePitch(const FAnimDatabaseFrameAttribute& A)
+{
+	if (!A.IsValid())
+	{
+		UE_LOGFMT(LogAnimDatabase, Error, "FrameAttributePitch: Invalid FrameAttribute.");
+		return FAnimDatabaseFrameAttribute();
+	}
+
+	if (A.Type != EAnimDatabaseAttributeType::Rotation)
+	{
+		UE_LOGFMT(LogAnimDatabase, Error, "FrameAttributePitch: FrameAttribute Types don't match. Got {Type}. Expected Rotation.", AttributeTypeNameInternal(A.Type));
+		return FAnimDatabaseFrameAttribute();
+	}
+
+	FAnimDatabaseFrameAttribute Out = MakeEmptyFrameAttribute();
+	Out.Type = EAnimDatabaseAttributeType::Angle;
+	UE::Learning::FrameAttribute::QuatPitch(*Out.FrameAttribute, *A.FrameAttribute);
+	return Out;
+}
+
+FAnimDatabaseFrameAttribute UAnimDatabaseFrameAttributeLibrary::FrameAttributeRoll(const FAnimDatabaseFrameAttribute& A)
+{
+	if (!A.IsValid())
+	{
+		UE_LOGFMT(LogAnimDatabase, Error, "FrameAttributeRoll: Invalid FrameAttribute.");
+		return FAnimDatabaseFrameAttribute();
+	}
+
+	if (A.Type != EAnimDatabaseAttributeType::Rotation)
+	{
+		UE_LOGFMT(LogAnimDatabase, Error, "FrameAttributeRoll: FrameAttribute Types don't match. Got {Type}. Expected Rotation.", AttributeTypeNameInternal(A.Type));
+		return FAnimDatabaseFrameAttribute();
+	}
+
+	FAnimDatabaseFrameAttribute Out = MakeEmptyFrameAttribute();
+	Out.Type = EAnimDatabaseAttributeType::Angle;
+	UE::Learning::FrameAttribute::QuatRoll(*Out.FrameAttribute, *A.FrameAttribute);
+	return Out;
+}
+
 FAnimDatabaseFrameAttribute UAnimDatabaseFrameAttributeLibrary::FrameAttributeTransformLocation(const FAnimDatabaseFrameAttribute& Transform)
 {
 	if (!Transform.IsValid())
@@ -7383,7 +8068,7 @@ FAnimDatabaseFrameAttribute UAnimDatabaseFrameAttributeLibrary::FrameAttributeTr
 
 	if (Transform.Type != EAnimDatabaseAttributeType::Transform)
 	{
-		UE_LOGFMT(LogAnimDatabase, Error, "FrameAttributeTransformLocation: FrameAttribute Types don't match. Got {Type}. Expected Transform.", *UEnum::GetValueAsString(Transform.Type));
+		UE_LOGFMT(LogAnimDatabase, Error, "FrameAttributeTransformLocation: FrameAttribute Types don't match. Got {Type}. Expected Transform.", AttributeTypeNameInternal(Transform.Type));
 		return FAnimDatabaseFrameAttribute();
 	}
 
@@ -7403,7 +8088,7 @@ FAnimDatabaseFrameAttribute UAnimDatabaseFrameAttributeLibrary::FrameAttributeTr
 
 	if (Transform.Type != EAnimDatabaseAttributeType::Transform)
 	{
-		UE_LOGFMT(LogAnimDatabase, Error, "FrameAttributeTransformRotation: FrameAttribute Types don't match. Got {Type}. Expected Transform.", *UEnum::GetValueAsString(Transform.Type));
+		UE_LOGFMT(LogAnimDatabase, Error, "FrameAttributeTransformRotation: FrameAttribute Types don't match. Got {Type}. Expected Transform.", AttributeTypeNameInternal(Transform.Type));
 		return FAnimDatabaseFrameAttribute();
 	}
 
@@ -7423,7 +8108,7 @@ FAnimDatabaseFrameAttribute UAnimDatabaseFrameAttributeLibrary::FrameAttributeTr
 
 	if (Transform.Type != EAnimDatabaseAttributeType::Transform)
 	{
-		UE_LOGFMT(LogAnimDatabase, Error, "FrameAttributeTransformScale: FrameAttribute Types don't match. Got {Type}. Expected Transform.", *UEnum::GetValueAsString(Transform.Type));
+		UE_LOGFMT(LogAnimDatabase, Error, "FrameAttributeTransformScale: FrameAttribute Types don't match. Got {Type}. Expected Transform.", AttributeTypeNameInternal(Transform.Type));
 		return FAnimDatabaseFrameAttribute();
 	}
 
@@ -7445,7 +8130,7 @@ FAnimDatabaseFrameAttribute UAnimDatabaseFrameAttributeLibrary::FrameAttributeEq
 
 	if (A.Type != B.Type)
 	{
-		UE_LOGFMT(LogAnimDatabase, Error, "FrameAttributeEqual: FrameAttribute Types don't match. Got {Type} and {Type}.", *UEnum::GetValueAsString(A.Type), *UEnum::GetValueAsString(B.Type));
+		UE_LOGFMT(LogAnimDatabase, Error, "FrameAttributeEqual: FrameAttribute Types don't match. Got {Type} and {Type}.", AttributeTypeNameInternal(A.Type), AttributeTypeNameInternal(B.Type));
 		return FAnimDatabaseFrameAttribute();
 	}
 
@@ -7467,7 +8152,7 @@ FAnimDatabaseFrameAttribute UAnimDatabaseFrameAttributeLibrary::FrameAttributeEq
 
 	if (TransformFrameAttribute.Type != EAnimDatabaseAttributeType::Transform)
 	{
-		UE_LOGFMT(LogAnimDatabase, Error, "FrameAttributeEqualTransform: FrameAttribute Types don't match. Got {Type}. Expected Transform.", *UEnum::GetValueAsString(TransformFrameAttribute.Type));
+		UE_LOGFMT(LogAnimDatabase, Error, "FrameAttributeEqualTransform: FrameAttribute Types don't match. Got {Type}. Expected Transform.", AttributeTypeNameInternal(TransformFrameAttribute.Type));
 		return FAnimDatabaseFrameAttribute();
 	}
 
@@ -7482,45 +8167,129 @@ FAnimDatabaseFrameAttribute UAnimDatabaseFrameAttributeLibrary::FrameAttributeEq
 	return Out;
 }
 
-FAnimDatabaseFrameAttribute UAnimDatabaseFrameAttributeLibrary::FrameAttributeGreaterThanFloat(const FAnimDatabaseFrameAttribute& BoolFrameAttribute, const float Threshold)
+FAnimDatabaseFrameAttribute UAnimDatabaseFrameAttributeLibrary::FrameAttributeGreaterThanFloat(const FAnimDatabaseFrameAttribute& FloatFrameAttribute, const float Threshold)
 {
 	TRACE_CPUPROFILER_EVENT_SCOPE(UAnimDatabaseFrameAttributeLibrary::FrameAttributeGreaterThanFloat);
 
-	if (!BoolFrameAttribute.IsValid())
+	if (!FloatFrameAttribute.IsValid())
 	{
 		UE_LOGFMT(LogAnimDatabase, Error, "FrameAttributeGreaterThanFloat: Invalid FrameAttribute.");
 		return FAnimDatabaseFrameAttribute();
 	}
 
-	if (BoolFrameAttribute.Type != EAnimDatabaseAttributeType::Float)
+	if (FloatFrameAttribute.Type != EAnimDatabaseAttributeType::Float)
 	{
-		UE_LOGFMT(LogAnimDatabase, Error, "FrameAttributeGreaterThanFloat: FrameAttribute Types don't match. Got {Type}. Expected Float.", *UEnum::GetValueAsString(BoolFrameAttribute.Type));
+		UE_LOGFMT(LogAnimDatabase, Error, "FrameAttributeGreaterThanFloat: FrameAttribute Types don't match. Got {Type}. Expected Float.", AttributeTypeNameInternal(FloatFrameAttribute.Type));
 		return FAnimDatabaseFrameAttribute();
 	}
 
 	FAnimDatabaseFrameAttribute Out = MakeEmptyFrameAttribute();
 	Out.Type = EAnimDatabaseAttributeType::Bool;
-	UE::Learning::FrameAttribute::GtConstant(*Out.FrameAttribute, *BoolFrameAttribute.FrameAttribute, { Threshold });
+	UE::Learning::FrameAttribute::GtConstant(*Out.FrameAttribute, *FloatFrameAttribute.FrameAttribute, { Threshold });
 	return Out;
 }
 
-FAnimDatabaseFrameAttribute UAnimDatabaseFrameAttributeLibrary::FrameAttributeLessThanFloat(const FAnimDatabaseFrameAttribute& BoolFrameAttribute, const float Threshold)
+FAnimDatabaseFrameAttribute UAnimDatabaseFrameAttributeLibrary::FrameAttributeLessThanFloat(const FAnimDatabaseFrameAttribute& FloatFrameAttribute, const float Threshold)
 {
-	if (!BoolFrameAttribute.IsValid())
+	if (!FloatFrameAttribute.IsValid())
 	{
 		UE_LOGFMT(LogAnimDatabase, Error, "FrameAttributeLessThanFloat: Invalid FrameAttribute.");
 		return FAnimDatabaseFrameAttribute();
 	}
 
-	if (BoolFrameAttribute.Type != EAnimDatabaseAttributeType::Float)
+	if (FloatFrameAttribute.Type != EAnimDatabaseAttributeType::Float)
+	{
+		UE_LOGFMT(LogAnimDatabase, Error, "FrameAttributeLessThanFloat: FrameAttribute Types don't match. Got {Type}. Expected Float.", AttributeTypeNameInternal(FloatFrameAttribute.Type));
+		return FAnimDatabaseFrameAttribute();
+	}
+
+	FAnimDatabaseFrameAttribute Out = MakeEmptyFrameAttribute();
+	Out.Type = EAnimDatabaseAttributeType::Bool;
+	UE::Learning::FrameAttribute::LtConstant(*Out.FrameAttribute, *FloatFrameAttribute.FrameAttribute, { Threshold });
+	return Out;
+}
+
+FAnimDatabaseFrameAttribute UAnimDatabaseFrameAttributeLibrary::FrameAttributeGreaterThanAngleDegrees(const FAnimDatabaseFrameAttribute& AngleFrameAttribute, const float Threshold)
+{
+	TRACE_CPUPROFILER_EVENT_SCOPE(UAnimDatabaseFrameAttributeLibrary::FrameAttributeGreaterThanAngleDegrees);
+
+	if (!AngleFrameAttribute.IsValid())
+	{
+		UE_LOGFMT(LogAnimDatabase, Error, "FrameAttributeGreaterThanAngleDegrees: Invalid FrameAttribute.");
+		return FAnimDatabaseFrameAttribute();
+	}
+
+	if (AngleFrameAttribute.Type != EAnimDatabaseAttributeType::Angle)
+	{
+		UE_LOGFMT(LogAnimDatabase, Error, "FrameAttributeGreaterThanAngleDegrees: FrameAttribute Types don't match. Got {Type}. Expected Angle.", AttributeTypeNameInternal(AngleFrameAttribute.Type));
+		return FAnimDatabaseFrameAttribute();
+	}
+
+	FAnimDatabaseFrameAttribute Out = MakeEmptyFrameAttribute();
+	Out.Type = EAnimDatabaseAttributeType::Bool;
+	UE::Learning::FrameAttribute::GtConstant(*Out.FrameAttribute, *AngleFrameAttribute.FrameAttribute, { FMath::DegreesToRadians(Threshold) });
+	return Out;
+}
+
+FAnimDatabaseFrameAttribute UAnimDatabaseFrameAttributeLibrary::FrameAttributeLessThanAngleDegrees(const FAnimDatabaseFrameAttribute& AngleFrameAttribute, const float Threshold)
+{
+	if (!AngleFrameAttribute.IsValid())
+	{
+		UE_LOGFMT(LogAnimDatabase, Error, "FrameAttributeLessThanAngleDegrees: Invalid FrameAttribute.");
+		return FAnimDatabaseFrameAttribute();
+	}
+
+	if (AngleFrameAttribute.Type != EAnimDatabaseAttributeType::Angle)
+	{
+		UE_LOGFMT(LogAnimDatabase, Error, "FrameAttributeLessThanAngleDegrees: FrameAttribute Types don't match. Got {Type}. Expected Angle.", AttributeTypeNameInternal(AngleFrameAttribute.Type));
+		return FAnimDatabaseFrameAttribute();
+	}
+
+	FAnimDatabaseFrameAttribute Out = MakeEmptyFrameAttribute();
+	Out.Type = EAnimDatabaseAttributeType::Bool;
+	UE::Learning::FrameAttribute::LtConstant(*Out.FrameAttribute, *AngleFrameAttribute.FrameAttribute, { FMath::DegreesToRadians(Threshold) });
+	return Out;
+}
+
+FAnimDatabaseFrameAttribute UAnimDatabaseFrameAttributeLibrary::FrameAttributeGreaterThanAngleRadians(const FAnimDatabaseFrameAttribute& AngleFrameAttribute, const float Threshold)
+{
+	TRACE_CPUPROFILER_EVENT_SCOPE(UAnimDatabaseFrameAttributeLibrary::FrameAttributeGreaterThanAngleRadians);
+
+	if (!AngleFrameAttribute.IsValid())
+	{
+		UE_LOGFMT(LogAnimDatabase, Error, "FrameAttributeGreaterThanAngleRadians: Invalid FrameAttribute.");
+		return FAnimDatabaseFrameAttribute();
+	}
+
+	if (AngleFrameAttribute.Type != EAnimDatabaseAttributeType::Angle)
+	{
+		UE_LOGFMT(LogAnimDatabase, Error, "FrameAttributeGreaterThanAngleRadians: FrameAttribute Types don't match. Got {Type}. Expected Angle.", AttributeTypeNameInternal(AngleFrameAttribute.Type));
+		return FAnimDatabaseFrameAttribute();
+	}
+
+	FAnimDatabaseFrameAttribute Out = MakeEmptyFrameAttribute();
+	Out.Type = EAnimDatabaseAttributeType::Bool;
+	UE::Learning::FrameAttribute::GtConstant(*Out.FrameAttribute, *AngleFrameAttribute.FrameAttribute, { Threshold });
+	return Out;
+}
+
+FAnimDatabaseFrameAttribute UAnimDatabaseFrameAttributeLibrary::FrameAttributeLessThanAngleRadians(const FAnimDatabaseFrameAttribute& AngleFrameAttribute, const float Threshold)
+{
+	if (!AngleFrameAttribute.IsValid())
+	{
+		UE_LOGFMT(LogAnimDatabase, Error, "FrameAttributeLessThanAngleRadians: Invalid FrameAttribute.");
+		return FAnimDatabaseFrameAttribute();
+	}
+
+	if (AngleFrameAttribute.Type != EAnimDatabaseAttributeType::Angle)
 	{
-		UE_LOGFMT(LogAnimDatabase, Error, "FrameAttributeLessThanFloat: FrameAttribute Types don't match. Got {Type}. Expected Float.", *UEnum::GetValueAsString(BoolFrameAttribute.Type));
+		UE_LOGFMT(LogAnimDatabase, Error, "FrameAttributeLessThanAngleRadians: FrameAttribute Types don't match. Got {Type}. Expected Angle.", AttributeTypeNameInternal(AngleFrameAttribute.Type));
 		return FAnimDatabaseFrameAttribute();
 	}
 
 	FAnimDatabaseFrameAttribute Out = MakeEmptyFrameAttribute();
 	Out.Type = EAnimDatabaseAttributeType::Bool;
-	UE::Learning::FrameAttribute::LtConstant(*Out.FrameAttribute, *BoolFrameAttribute.FrameAttribute, { Threshold });
+	UE::Learning::FrameAttribute::LtConstant(*Out.FrameAttribute, *AngleFrameAttribute.FrameAttribute, { Threshold });
 	return Out;
 }
 
@@ -7534,7 +8303,7 @@ FAnimDatabaseFrameAttribute UAnimDatabaseFrameAttributeLibrary::FrameAttributeBo
 
 	if (BoolFrameAttribute.Type != EAnimDatabaseAttributeType::Bool)
 	{
-		UE_LOGFMT(LogAnimDatabase, Error, "FrameAttributeBoolToFloat: FrameAttribute Types don't match. Got {Type}. Expected Bool.", *UEnum::GetValueAsString(BoolFrameAttribute.Type));
+		UE_LOGFMT(LogAnimDatabase, Error, "FrameAttributeBoolToFloat: FrameAttribute Types don't match. Got {Type}. Expected Bool.", AttributeTypeNameInternal(BoolFrameAttribute.Type));
 		return FAnimDatabaseFrameAttribute();
 	}
 
@@ -7544,6 +8313,46 @@ FAnimDatabaseFrameAttribute UAnimDatabaseFrameAttributeLibrary::FrameAttributeBo
 	return Out;
 }
 
+FAnimDatabaseFrameAttribute UAnimDatabaseFrameAttributeLibrary::FrameAttributeAngleToFloatDegrees(const FAnimDatabaseFrameAttribute& AngleFrameAttribute)
+{
+	if (!AngleFrameAttribute.IsValid())
+	{
+		UE_LOGFMT(LogAnimDatabase, Error, "FrameAttributeAngleToFloatDegrees: Invalid FrameAttribute.");
+		return FAnimDatabaseFrameAttribute();
+	}
+
+	if (AngleFrameAttribute.Type != EAnimDatabaseAttributeType::Angle)
+	{
+		UE_LOGFMT(LogAnimDatabase, Error, "FrameAttributeAngleToFloatDegrees: FrameAttribute Types don't match. Got {Type}. Expected Angle.", AttributeTypeNameInternal(AngleFrameAttribute.Type));
+		return FAnimDatabaseFrameAttribute();
+	}
+
+	FAnimDatabaseFrameAttribute Out = FrameAttributeCopy(AngleFrameAttribute);
+	Out.Type = EAnimDatabaseAttributeType::Float;
+	UE::Learning::FrameAttribute::MulConstantInplace(*Out.FrameAttribute, { 180.0f / UE_PI });
+	return Out;
+}
+
+FAnimDatabaseFrameAttribute UAnimDatabaseFrameAttributeLibrary::FrameAttributeAngleToFloatRadians(const FAnimDatabaseFrameAttribute& AngleFrameAttribute)
+{
+	if (!AngleFrameAttribute.IsValid())
+	{
+		UE_LOGFMT(LogAnimDatabase, Error, "FrameAttributeAngleToFloatRadians: Invalid FrameAttribute.");
+		return FAnimDatabaseFrameAttribute();
+	}
+
+	if (AngleFrameAttribute.Type != EAnimDatabaseAttributeType::Angle)
+	{
+		UE_LOGFMT(LogAnimDatabase, Error, "FrameAttributeAngleToFloatRadians: FrameAttribute Types don't match. Got {Type}. Expected Angle.", AttributeTypeNameInternal(AngleFrameAttribute.Type));
+		return FAnimDatabaseFrameAttribute();
+	}
+
+	// We can just return it directly with a different type the underlying data doesn't need changing
+	FAnimDatabaseFrameAttribute Out = AngleFrameAttribute;
+	Out.Type = EAnimDatabaseAttributeType::Float;
+	return Out;
+}
+
 FAnimDatabaseFrameAttribute UAnimDatabaseFrameAttributeLibrary::FrameAttributeFloatToBool(const FAnimDatabaseFrameAttribute& A)
 {
 	if (!A.IsValid())
@@ -7554,7 +8363,7 @@ FAnimDatabaseFrameAttribute UAnimDatabaseFrameAttributeLibrary::FrameAttributeFl
 
 	if (A.Type != EAnimDatabaseAttributeType::Float)
 	{
-		UE_LOGFMT(LogAnimDatabase, Error, "FrameAttributeFloatToBool: FrameAttribute Types don't match. Got {Type}. Expected Float.", *UEnum::GetValueAsString(A.Type));
+		UE_LOGFMT(LogAnimDatabase, Error, "FrameAttributeFloatToBool: FrameAttribute Types don't match. Got {Type}. Expected Float.", AttributeTypeNameInternal(A.Type));
 		return FAnimDatabaseFrameAttribute();
 	}
 
@@ -7564,6 +8373,46 @@ FAnimDatabaseFrameAttribute UAnimDatabaseFrameAttributeLibrary::FrameAttributeFl
 	return Out;
 }
 
+FAnimDatabaseFrameAttribute UAnimDatabaseFrameAttributeLibrary::FrameAttributeFloatToAngleDegrees(const FAnimDatabaseFrameAttribute& FloatFrameAttribute)
+{
+	if (!FloatFrameAttribute.IsValid())
+	{
+		UE_LOGFMT(LogAnimDatabase, Error, "FrameAttributeFloatToAngleDegrees: Invalid FrameAttribute.");
+		return FAnimDatabaseFrameAttribute();
+	}
+
+	if (FloatFrameAttribute.Type != EAnimDatabaseAttributeType::Float)
+	{
+		UE_LOGFMT(LogAnimDatabase, Error, "FrameAttributeFloatToAngleDegrees: FrameAttribute Types don't match. Got {Type}. Expected Float.", AttributeTypeNameInternal(FloatFrameAttribute.Type));
+		return FAnimDatabaseFrameAttribute();
+	}
+
+	FAnimDatabaseFrameAttribute Out = FrameAttributeCopy(FloatFrameAttribute);
+	Out.Type = EAnimDatabaseAttributeType::Angle;
+	UE::Learning::FrameAttribute::MulConstantInplace(*Out.FrameAttribute, { UE_PI / 180.0f });
+	return Out;
+}
+
+FAnimDatabaseFrameAttribute UAnimDatabaseFrameAttributeLibrary::FrameAttributeFloatToAngleRadians(const FAnimDatabaseFrameAttribute& FloatFrameAttribute)
+{
+	if (!FloatFrameAttribute.IsValid())
+	{
+		UE_LOGFMT(LogAnimDatabase, Error, "FrameAttributeFloatToAngleRadians: Invalid FrameAttribute.");
+		return FAnimDatabaseFrameAttribute();
+	}
+
+	if (FloatFrameAttribute.Type != EAnimDatabaseAttributeType::Float)
+	{
+		UE_LOGFMT(LogAnimDatabase, Error, "FrameAttributeFloatToAngleRadians: FrameAttribute Types don't match. Got {Type}. Expected Float.", AttributeTypeNameInternal(FloatFrameAttribute.Type));
+		return FAnimDatabaseFrameAttribute();
+	}
+
+	// We can just return it directly with a different type the underlying data doesn't need changing
+	FAnimDatabaseFrameAttribute Out = FloatFrameAttribute;
+	Out.Type = EAnimDatabaseAttributeType::Angle;
+	return Out;
+}
+
 FAnimDatabaseFrameAttribute UAnimDatabaseFrameAttributeLibrary::FrameAttributeFilterGaussian(const FAnimDatabaseFrameAttribute& A, const float StandardDeviationInFrames)
 {
 	if (!A.IsValid())
@@ -7574,7 +8423,7 @@ FAnimDatabaseFrameAttribute UAnimDatabaseFrameAttributeLibrary::FrameAttributeFi
 
 	if (A.Type != EAnimDatabaseAttributeType::Float && A.Type != EAnimDatabaseAttributeType::Location && A.Type != EAnimDatabaseAttributeType::Direction)
 	{
-		UE_LOGFMT(LogAnimDatabase, Error, "FrameAttributeFilterGaussian: FrameAttribute Types don't match. Got {Type}. Expected Float, Location, or Direction.", *UEnum::GetValueAsString(A.Type));
+		UE_LOGFMT(LogAnimDatabase, Error, "FrameAttributeFilterGaussian: FrameAttribute Types don't match. Got {Type}. Expected Float, Location, or Direction.", AttributeTypeNameInternal(A.Type));
 		return FAnimDatabaseFrameAttribute();
 	}
 
@@ -7618,7 +8467,7 @@ FAnimDatabaseFrameAttribute UAnimDatabaseFrameAttributeLibrary::FrameAttributeFi
 
 	if (A.Type != EAnimDatabaseAttributeType::Float && A.Type != EAnimDatabaseAttributeType::Location && A.Type != EAnimDatabaseAttributeType::Direction)
 	{
-		UE_LOGFMT(LogAnimDatabase, Error, "FrameAttributeFilterSavGol: FrameAttribute Types don't match. Got {Type}. Expected Float, Location, or Direction.", *UEnum::GetValueAsString(A.Type));
+		UE_LOGFMT(LogAnimDatabase, Error, "FrameAttributeFilterSavGol: FrameAttribute Types don't match. Got {Type}. Expected Float, Location, or Direction.", AttributeTypeNameInternal(A.Type));
 		return FAnimDatabaseFrameAttribute();
 	}
 
@@ -7668,7 +8517,7 @@ FAnimDatabaseFrameAttribute UAnimDatabaseFrameAttributeLibrary::FrameAttributeFi
 
 	if (A.Type != EAnimDatabaseAttributeType::Bool)
 	{
-		UE_LOGFMT(LogAnimDatabase, Error, "FrameAttributeFilterMajorityVote: FrameAttribute Types don't match. Got {Type}. Expected Bool.", *UEnum::GetValueAsString(A.Type));
+		UE_LOGFMT(LogAnimDatabase, Error, "FrameAttributeFilterMajorityVote: FrameAttribute Types don't match. Got {Type}. Expected Bool.", AttributeTypeNameInternal(A.Type));
 		return FAnimDatabaseFrameAttribute();
 	}
 
@@ -7690,7 +8539,7 @@ FAnimDatabaseFrameRanges UAnimDatabaseFrameAttributeLibrary::FrameAttributeBoolT
 
 	if (BoolFrameAttribute.Type != EAnimDatabaseAttributeType::Bool)
 	{
-		UE_LOGFMT(LogAnimDatabase, Error, "FrameAttributeBoolToFrameRanges: FrameAttribute Types don't match. Got {Type}. Expected Bool.", *UEnum::GetValueAsString(BoolFrameAttribute.Type));
+		UE_LOGFMT(LogAnimDatabase, Error, "FrameAttributeBoolToFrameRanges: FrameAttribute Types don't match. Got {Type}. Expected Bool.", AttributeTypeNameInternal(BoolFrameAttribute.Type));
 		return FAnimDatabaseFrameRanges();
 	}
 
@@ -7711,13 +8560,13 @@ FAnimDatabaseFrameAttribute UAnimDatabaseFrameAttributeLibrary::FrameAttributeSe
 
 	if (Cond.Type != EAnimDatabaseAttributeType::Bool)
 	{
-		UE_LOGFMT(LogAnimDatabase, Error, "FrameAttributeSelect: FrameAttribute Types don't match. Got {Type}. Expected Bool.", *UEnum::GetValueAsString(Cond.Type));
+		UE_LOGFMT(LogAnimDatabase, Error, "FrameAttributeSelect: FrameAttribute Types don't match. Got {Type}. Expected Bool.", AttributeTypeNameInternal(Cond.Type));
 		return FAnimDatabaseFrameAttribute();
 	}
 
 	if (True.Type != False.Type)
 	{
-		UE_LOGFMT(LogAnimDatabase, Error, "FrameAttributeSelect: FrameAttribute Types don't match. Got {Type0} and {Type1}.", *UEnum::GetValueAsString(True.Type), *UEnum::GetValueAsString(False.Type));
+		UE_LOGFMT(LogAnimDatabase, Error, "FrameAttributeSelect: FrameAttribute Types don't match. Got {Type0} and {Type1}.", AttributeTypeNameInternal(True.Type), AttributeTypeNameInternal(False.Type));
 		return FAnimDatabaseFrameAttribute();
 	}
 
@@ -7746,7 +8595,7 @@ FAnimDatabaseFrameAttribute UAnimDatabaseFrameAttributeLibrary::FrameAttributeNo
 
 	if (A.Type != EAnimDatabaseAttributeType::Bool)
 	{
-		UE_LOGFMT(LogAnimDatabase, Error, "FrameAttributeNot: FrameAttribute Types don't match. Got {Type}. Expected Bool.", *UEnum::GetValueAsString(A.Type));
+		UE_LOGFMT(LogAnimDatabase, Error, "FrameAttributeNot: FrameAttribute Types don't match. Got {Type}. Expected Bool.", AttributeTypeNameInternal(A.Type));
 		return FAnimDatabaseFrameAttribute();
 	}
 
@@ -7766,7 +8615,7 @@ FAnimDatabaseFrameAttribute UAnimDatabaseFrameAttributeLibrary::FrameAttributeAn
 
 	if (A.Type != EAnimDatabaseAttributeType::Bool || B.Type != EAnimDatabaseAttributeType::Bool)
 	{
-		UE_LOGFMT(LogAnimDatabase, Error, "FrameAttributeAnd: FrameAttribute Types don't match. Got {Type0} and {Type1}. Expected Bool.", *UEnum::GetValueAsString(A.Type), *UEnum::GetValueAsString(B.Type));
+		UE_LOGFMT(LogAnimDatabase, Error, "FrameAttributeAnd: FrameAttribute Types don't match. Got {Type0} and {Type1}. Expected Bool.", AttributeTypeNameInternal(A.Type), AttributeTypeNameInternal(B.Type));
 		return FAnimDatabaseFrameAttribute();
 	}
 
@@ -7786,7 +8635,7 @@ FAnimDatabaseFrameAttribute UAnimDatabaseFrameAttributeLibrary::FrameAttributeOr
 
 	if (A.Type != EAnimDatabaseAttributeType::Bool || B.Type != EAnimDatabaseAttributeType::Bool)
 	{
-		UE_LOGFMT(LogAnimDatabase, Error, "FrameAttributeOr: FrameAttribute Types don't match. Got {Type0} and {Type1}. Expected Bool.", *UEnum::GetValueAsString(A.Type), *UEnum::GetValueAsString(B.Type));
+		UE_LOGFMT(LogAnimDatabase, Error, "FrameAttributeOr: FrameAttribute Types don't match. Got {Type0} and {Type1}. Expected Bool.", AttributeTypeNameInternal(A.Type), AttributeTypeNameInternal(B.Type));
 		return FAnimDatabaseFrameAttribute();
 	}
 
@@ -7910,7 +8759,7 @@ void UAnimDatabaseFrameAttributeLibrary::FrameAttributeLocationMeanAndStd(FVecto
 
 	if (LocationFrameAttribute.Type != EAnimDatabaseAttributeType::Location)
 	{
-		UE_LOGFMT(LogAnimDatabase, Error, "FrameAttributeLocationMeanAndStd: FrameAttribute Types don't match. Got {Type}. Expected Location.", *UEnum::GetValueAsString(LocationFrameAttribute.Type));
+		UE_LOGFMT(LogAnimDatabase, Error, "FrameAttributeLocationMeanAndStd: FrameAttribute Types don't match. Got {Type}. Expected Location.", AttributeTypeNameInternal(LocationFrameAttribute.Type));
 		OutMean = FVector::ZeroVector;
 		OutStd = FVector::ZeroVector;
 		return;
@@ -7935,7 +8784,7 @@ void UAnimDatabaseFrameAttributeLibrary::FrameAttributeFloatAverage(float& OutAv
 
 	if (FloatFrameAttribute.Type != EAnimDatabaseAttributeType::Float)
 	{
-		UE_LOGFMT(LogAnimDatabase, Error, "FrameAttributeFloatAverage: FrameAttribute Types don't match. Got {Type}. Expected Float.", *UEnum::GetValueAsString(FloatFrameAttribute.Type));
+		UE_LOGFMT(LogAnimDatabase, Error, "FrameAttributeFloatAverage: FrameAttribute Types don't match. Got {Type}. Expected Float.", AttributeTypeNameInternal(FloatFrameAttribute.Type));
 		OutAverageFloat = 0.0f;
 		return;
 	}
@@ -7943,6 +8792,50 @@ void UAnimDatabaseFrameAttributeLibrary::FrameAttributeFloatAverage(float& OutAv
 	UE::Learning::FrameAttribute::Mean(MakeArrayView(&OutAverageFloat, 1), *FloatFrameAttribute.FrameAttribute);
 }
 
+void UAnimDatabaseFrameAttributeLibrary::FrameAttributeAngleAverageDegrees(float& OutAverageAngle, const FAnimDatabaseFrameAttribute& AngleFrameAttribute)
+{
+	TRACE_CPUPROFILER_EVENT_SCOPE(UAnimDatabaseFrameAttributeLibrary::FrameAttributeAngleAverageDegrees);
+
+	if (!AngleFrameAttribute.IsValid())
+	{
+		UE_LOGFMT(LogAnimDatabase, Error, "FrameAttributeAngleAverageDegrees: Invalid FrameAttribute.");
+		OutAverageAngle = 0.0f;
+		return;
+	}
+
+	if (AngleFrameAttribute.Type != EAnimDatabaseAttributeType::Angle)
+	{
+		UE_LOGFMT(LogAnimDatabase, Error, "FrameAttributeAngleAverageDegrees: FrameAttribute Types don't match. Got {Type}. Expected Angle.", AttributeTypeNameInternal(AngleFrameAttribute.Type));
+		OutAverageAngle = 0.0f;
+		return;
+	}
+
+	UE::Learning::FrameAttribute::AngularMean(MakeArrayView(&OutAverageAngle, 1), *AngleFrameAttribute.FrameAttribute);
+	OutAverageAngle = FMath::RadiansToDegrees(OutAverageAngle);
+}
+
+void UAnimDatabaseFrameAttributeLibrary::FrameAttributeAngleAverageRadians(float& OutAverageAngle, const FAnimDatabaseFrameAttribute& AngleFrameAttribute)
+{
+	TRACE_CPUPROFILER_EVENT_SCOPE(UAnimDatabaseFrameAttributeLibrary::FrameAttributeAngleAverageRadians);
+
+	if (!AngleFrameAttribute.IsValid())
+	{
+		UE_LOGFMT(LogAnimDatabase, Error, "FrameAttributeAngleAverageRadians: Invalid FrameAttribute.");
+		OutAverageAngle = 0.0f;
+		return;
+	}
+
+	if (AngleFrameAttribute.Type != EAnimDatabaseAttributeType::Angle)
+	{
+		UE_LOGFMT(LogAnimDatabase, Error, "FrameAttributeAngleAverageRadians: FrameAttribute Types don't match. Got {Type}. Expected Angle.", AttributeTypeNameInternal(AngleFrameAttribute.Type));
+		OutAverageAngle = 0.0f;
+		return;
+	}
+
+	UE::Learning::FrameAttribute::AngularMean(MakeArrayView(&OutAverageAngle, 1), *AngleFrameAttribute.FrameAttribute);
+}
+
+
 void UAnimDatabaseFrameAttributeLibrary::FrameAttributeFloatMinAndMax(float& OutMin, float& OutMax, const FAnimDatabaseFrameAttribute& FloatFrameAttribute)
 {
 	TRACE_CPUPROFILER_EVENT_SCOPE(UAnimDatabaseFrameAttributeLibrary::FrameAttributeFloatMinAndMax);
@@ -7957,7 +8850,7 @@ void UAnimDatabaseFrameAttributeLibrary::FrameAttributeFloatMinAndMax(float& Out
 
 	if (FloatFrameAttribute.Type != EAnimDatabaseAttributeType::Float)
 	{
-		UE_LOGFMT(LogAnimDatabase, Error, "FrameAttributeFloatMinAndMax: FrameAttribute Types don't match. Got {Type}. Expected Float.", *UEnum::GetValueAsString(FloatFrameAttribute.Type));
+		UE_LOGFMT(LogAnimDatabase, Error, "FrameAttributeFloatMinAndMax: FrameAttribute Types don't match. Got {Type}. Expected Float.", AttributeTypeNameInternal(FloatFrameAttribute.Type));
 		OutMin = 0.0f;
 		OutMax = 0.0f;
 		return;
@@ -7966,6 +8859,55 @@ void UAnimDatabaseFrameAttributeLibrary::FrameAttributeFloatMinAndMax(float& Out
 	UE::Learning::FrameAttribute::MinMax(MakeArrayView(&OutMin, 1), MakeArrayView(&OutMax, 1), *FloatFrameAttribute.FrameAttribute);
 }
 
+void UAnimDatabaseFrameAttributeLibrary::FrameAttributeAngleMinAndMaxDegrees(float& OutMin, float& OutMax, const FAnimDatabaseFrameAttribute& AngleFrameAttribute)
+{
+	TRACE_CPUPROFILER_EVENT_SCOPE(UAnimDatabaseFrameAttributeLibrary::FrameAttributeAngleMinAndMaxDegrees);
+
+	if (!AngleFrameAttribute.IsValid())
+	{
+		UE_LOGFMT(LogAnimDatabase, Error, "FrameAttributeAngleMinAndMaxDegrees: Invalid FrameAttribute.");
+		OutMin = 0.0f;
+		OutMax = 0.0f;
+		return;
+	}
+
+	if (AngleFrameAttribute.Type != EAnimDatabaseAttributeType::Angle)
+	{
+		UE_LOGFMT(LogAnimDatabase, Error, "FrameAttributeAngleMinAndMaxDegrees: FrameAttribute Types don't match. Got {Type}. Expected Angle.", AttributeTypeNameInternal(AngleFrameAttribute.Type));
+		OutMin = 0.0f;
+		OutMax = 0.0f;
+		return;
+	}
+
+	UE::Learning::FrameAttribute::MinMax(MakeArrayView(&OutMin, 1), MakeArrayView(&OutMax, 1), *AngleFrameAttribute.FrameAttribute);
+	OutMin = FMath::RadiansToDegrees(OutMin);
+	OutMax = FMath::RadiansToDegrees(OutMax);
+}
+
+void UAnimDatabaseFrameAttributeLibrary::FrameAttributeAngleMinAndMaxRadians(float& OutMin, float& OutMax, const FAnimDatabaseFrameAttribute& AngleFrameAttribute)
+{
+	TRACE_CPUPROFILER_EVENT_SCOPE(UAnimDatabaseFrameAttributeLibrary::FrameAttributeAngleMinAndMaxRadians);
+
+	if (!AngleFrameAttribute.IsValid())
+	{
+		UE_LOGFMT(LogAnimDatabase, Error, "FrameAttributeAngleMinAndMaxRadians: Invalid FrameAttribute.");
+		OutMin = 0.0f;
+		OutMax = 0.0f;
+		return;
+	}
+
+	if (AngleFrameAttribute.Type != EAnimDatabaseAttributeType::Angle)
+	{
+		UE_LOGFMT(LogAnimDatabase, Error, "FrameAttributeAngleMinAndMaxRadians: FrameAttribute Types don't match. Got {Type}. Expected Angle.", AttributeTypeNameInternal(AngleFrameAttribute.Type));
+		OutMin = 0.0f;
+		OutMax = 0.0f;
+		return;
+	}
+
+	UE::Learning::FrameAttribute::MinMax(MakeArrayView(&OutMin, 1), MakeArrayView(&OutMax, 1), *AngleFrameAttribute.FrameAttribute);
+}
+
+
 void UAnimDatabaseFrameAttributeLibrary::FrameAttributeLocationAverage(FVector& OutAverageLocation, const FAnimDatabaseFrameAttribute& LocationFrameAttribute)
 {
 	TRACE_CPUPROFILER_EVENT_SCOPE(UAnimDatabaseFrameAttributeLibrary::FrameAttributeLocationAverage);
@@ -7979,7 +8921,7 @@ void UAnimDatabaseFrameAttributeLibrary::FrameAttributeLocationAverage(FVector&
 
 	if (LocationFrameAttribute.Type != EAnimDatabaseAttributeType::Location)
 	{
-		UE_LOGFMT(LogAnimDatabase, Error, "FrameAttributeLocationAverage: FrameAttribute Types don't match. Got {Type}. Expected Location.", *UEnum::GetValueAsString(LocationFrameAttribute.Type));
+		UE_LOGFMT(LogAnimDatabase, Error, "FrameAttributeLocationAverage: FrameAttribute Types don't match. Got {Type}. Expected Location.", AttributeTypeNameInternal(LocationFrameAttribute.Type));
 		OutAverageLocation = FVector::ZeroVector;
 		return;
 	}
@@ -8002,7 +8944,7 @@ void UAnimDatabaseFrameAttributeLibrary::FrameAttributeDirectionAverage(FVector&
 
 	if (DirectionFrameAttribute.Type != EAnimDatabaseAttributeType::Direction)
 	{
-		UE_LOGFMT(LogAnimDatabase, Error, "FrameAttributeDirectionAverage: FrameAttribute Types don't match. Got {Type}. Expected Direction.", *UEnum::GetValueAsString(DirectionFrameAttribute.Type));
+		UE_LOGFMT(LogAnimDatabase, Error, "FrameAttributeDirectionAverage: FrameAttribute Types don't match. Got {Type}. Expected Direction.", AttributeTypeNameInternal(DirectionFrameAttribute.Type));
 		OutAverageDirection = FVector::ForwardVector;
 		return;
 	}
@@ -8033,7 +8975,7 @@ void UAnimDatabaseFrameAttributeLibrary::FrameAttributeRotationMeanAndStdAsQuat(
 
 	if (RotationFrameAttribute.Type != EAnimDatabaseAttributeType::Rotation)
 	{
-		UE_LOGFMT(LogAnimDatabase, Error, "FrameAttributeRotationMeanAndStdAsQuat: FrameAttribute Types don't match. Got {Type}. Expected Rotation.", *UEnum::GetValueAsString(RotationFrameAttribute.Type));
+		UE_LOGFMT(LogAnimDatabase, Error, "FrameAttributeRotationMeanAndStdAsQuat: FrameAttribute Types don't match. Got {Type}. Expected Rotation.", AttributeTypeNameInternal(RotationFrameAttribute.Type));
 		OutMean = FQuat::Identity;
 		OutStd = FVector::ZeroVector;
 		return;
@@ -8059,7 +9001,7 @@ void UAnimDatabaseFrameAttributeLibrary::FrameAttributeScaleMeanAndStd(FVector&
 
 	if (ScaleFrameAttribute.Type != EAnimDatabaseAttributeType::Scale)
 	{
-		UE_LOGFMT(LogAnimDatabase, Error, "FrameAttributeScaleMeanAndStd: FrameAttribute Types don't match. Got {Type}. Expected Scale.", *UEnum::GetValueAsString(ScaleFrameAttribute.Type));
+		UE_LOGFMT(LogAnimDatabase, Error, "FrameAttributeScaleMeanAndStd: FrameAttribute Types don't match. Got {Type}. Expected Scale.", AttributeTypeNameInternal(ScaleFrameAttribute.Type));
 		OutMean = FVector::OneVector;
 		OutLogStd = FVector::ZeroVector;
 		return;
@@ -8087,7 +9029,7 @@ void UAnimDatabaseFrameAttributeLibrary::FrameAttributeVelocityMeanAndStd(FVecto
 		VelocityFrameAttribute.Type != EAnimDatabaseAttributeType::AngularVelocity &&
 		VelocityFrameAttribute.Type != EAnimDatabaseAttributeType::ScalarVelocity)
 	{
-		UE_LOGFMT(LogAnimDatabase, Error, "FrameAttributeVelocityMeanAndStd: FrameAttribute Types don't match. Got {Type}. Expected Velocity.", *UEnum::GetValueAsString(VelocityFrameAttribute.Type));
+		UE_LOGFMT(LogAnimDatabase, Error, "FrameAttributeVelocityMeanAndStd: FrameAttribute Types don't match. Got {Type}. Expected Velocity.", AttributeTypeNameInternal(VelocityFrameAttribute.Type));
 		OutMean = FVector::ZeroVector;
 		OutStd = FVector::ZeroVector;
 		return;
@@ -8115,7 +9057,7 @@ void UAnimDatabaseFrameAttributeLibrary::FrameAttributeLocationDistanceTraveled(
 
 	if (LocationFrameAttribute.Type != EAnimDatabaseAttributeType::Location)
 	{
-		UE_LOGFMT(LogAnimDatabase, Error, "FrameAttributeLocationDistanceTraveled: FrameAttribute Types don't match. Got {Type}. Expected Location.", *UEnum::GetValueAsString(LocationFrameAttribute.Type));
+		UE_LOGFMT(LogAnimDatabase, Error, "FrameAttributeLocationDistanceTraveled: FrameAttribute Types don't match. Got {Type}. Expected Location.", AttributeTypeNameInternal(LocationFrameAttribute.Type));
 		OutMean = 0.0f;
 		OutStd = 0.0f;
 		OutMin = 0.0f;
@@ -8158,7 +9100,7 @@ void UAnimDatabaseFrameAttributeLibrary::FrameAttributeRotationAngleTraveled(flo
 
 	if (RotationFrameAttribute.Type != EAnimDatabaseAttributeType::Rotation)
 	{
-		UE_LOGFMT(LogAnimDatabase, Error, "FrameAttributeRotationAngleTraveled: FrameAttribute Types don't match. Got {Type}. Expected Rotation.", *UEnum::GetValueAsString(RotationFrameAttribute.Type));
+		UE_LOGFMT(LogAnimDatabase, Error, "FrameAttributeRotationAngleTraveled: FrameAttribute Types don't match. Got {Type}. Expected Rotation.", AttributeTypeNameInternal(RotationFrameAttribute.Type));
 		OutMean = 0.0f;
 		OutStd = 0.0f;
 		OutMin = 0.0f;
@@ -8190,6 +9132,181 @@ void UAnimDatabaseFrameAttributeLibrary::FrameAttributeRotationAngleTraveled(flo
 	OutStd = FMath::RadiansToDegrees(OutStd);
 }
 
+FString UAnimDatabaseFrameAttributeLibrary::FrameAttributeToString(const FAnimDatabaseFrameAttribute& FrameAttribute)
+{
+	if (!FrameAttribute.IsValid()) { return TEXT("Invalid"); }
+
+	return FString::Printf(TEXT("Type=%s EntryNum=%i TotalRangeNum=%i TotalFrameNum=%i"), 
+		AttributeTypeNameInternal(FrameAttribute.Type),
+		FrameAttribute.FrameAttribute->FrameRangeSet.GetEntryNum(), 
+		FrameAttribute.FrameAttribute->FrameRangeSet.GetTotalRangeNum(),
+		FrameAttribute.FrameAttribute->FrameRangeSet.GetTotalFrameNum());
+}
+
+
+FString UAnimDatabaseFrameAttributeLibrary::FrameAttributeToStringFormat(const UAnimDatabase* Database, const FAnimDatabaseFrameAttribute& FrameAttribute, const int32 Cutoff)
+{
+	if (!FrameAttribute.IsValid()) { return TEXT("Invalid"); }
+
+	const int32 RoundedCuttoff = (Cutoff / 2) * 2 + 1;
+
+	FString Output = FString::Printf(TEXT("Type=%s EntryNum=%i TotalRangeNum=%i TotalFrameNum=%i Data=[\n"),
+		AttributeTypeNameInternal(FrameAttribute.Type),
+		FrameAttribute.FrameAttribute->FrameRangeSet.GetEntryNum(),
+		FrameAttribute.FrameAttribute->FrameRangeSet.GetTotalRangeNum(),
+		FrameAttribute.FrameAttribute->FrameRangeSet.GetTotalFrameNum());
+
+	const int32 EntryNum = FrameAttribute.FrameAttribute->FrameRangeSet.GetEntryNum();
+	for (int32 EntryIdx = 0; EntryIdx < EntryNum; EntryIdx++)
+	{
+		const int32 SequenceIdx = FrameAttribute.FrameAttribute->FrameRangeSet.GetEntrySequence(EntryIdx);
+		const int32 RangeNum = FrameAttribute.FrameAttribute->FrameRangeSet.GetEntryRangeNum(EntryIdx);
+
+		if (Database)
+		{
+			Output += FString::Printf(TEXT("   %s%s [\n"), 
+				*Database->GetSequenceAssetName(SequenceIdx), 
+				Database->GetIsMirrored(SequenceIdx) ? TEXT(" (mirrored)") : TEXT(""));
+		}
+		else
+		{
+			Output += FString::Printf(TEXT("   Sequence %i [\n"), SequenceIdx);
+
+		}
+
+		if (Cutoff < 0 || RangeNum < RoundedCuttoff)
+		{
+			for (int32 RangeIdx = 0; RangeIdx < RangeNum; RangeIdx++)
+			{
+				Output += FString::Printf(TEXT("        %i:%i "),
+					FrameAttribute.FrameAttribute->FrameRangeSet.GetEntryRangeStart(EntryIdx, RangeIdx),
+					FrameAttribute.FrameAttribute->FrameRangeSet.GetEntryRangeStop(EntryIdx, RangeIdx));
+
+				const int32 FrameNum = FrameAttribute.FrameAttribute->FrameRangeSet.GetEntryRangeLength(EntryIdx, RangeIdx);
+				const int32 FrameOffset = FrameAttribute.FrameAttribute->FrameRangeSet.GetEntryRangeOffset(EntryIdx, RangeIdx);
+
+				if (Cutoff < 0 || FrameNum < RoundedCuttoff)
+				{
+					Output += TEXT("[");
+					for (int32 FrameIdx = 0; FrameIdx < FrameNum; FrameIdx++)
+					{
+						Output += UE::AnimDatabase::FrameAttribute::Private::FrameAttributeFrameToString(FrameAttribute, FrameOffset + FrameIdx);
+						if (FrameIdx != FrameNum - 1) { Output += TEXT(" "); }
+					}
+					Output += TEXT("]\n");
+				}
+				else
+				{
+					Output += TEXT("[");
+					for (int32 FrameIdx = 0; FrameIdx < RoundedCuttoff / 2; FrameIdx++)
+					{
+						Output += UE::AnimDatabase::FrameAttribute::Private::FrameAttributeFrameToString(FrameAttribute, FrameOffset + FrameIdx);
+						Output += TEXT(" ");
+					}
+
+					Output += TEXT("... ");
+
+					for (int32 FrameIdx = FrameNum - RoundedCuttoff / 2 - 1; FrameIdx < FrameNum; FrameIdx++)
+					{
+						Output += UE::AnimDatabase::FrameAttribute::Private::FrameAttributeFrameToString(FrameAttribute, FrameOffset + FrameIdx);
+						if (FrameIdx != FrameNum - 1) { Output += TEXT(" "); }
+					}
+
+					Output += TEXT("]\n");
+				}
+			}
+		}
+		else
+		{
+			for (int32 RangeIdx = 0; RangeIdx < RoundedCuttoff / 2; RangeIdx++)
+			{
+				Output += FString::Printf(TEXT("        %i:%i "),
+					FrameAttribute.FrameAttribute->FrameRangeSet.GetEntryRangeStart(EntryIdx, RangeIdx),
+					FrameAttribute.FrameAttribute->FrameRangeSet.GetEntryRangeStop(EntryIdx, RangeIdx));
+
+				const int32 FrameNum = FrameAttribute.FrameAttribute->FrameRangeSet.GetEntryRangeLength(EntryIdx, RangeIdx);
+				const int32 FrameOffset = FrameAttribute.FrameAttribute->FrameRangeSet.GetEntryRangeOffset(EntryIdx, RangeIdx);
+
+				if (FrameNum < RoundedCuttoff)
+				{
+					Output += TEXT("[");
+					for (int32 FrameIdx = 0; FrameIdx < FrameNum; FrameIdx++)
+					{
+						Output += UE::AnimDatabase::FrameAttribute::Private::FrameAttributeFrameToString(FrameAttribute, FrameOffset + FrameIdx);
+						if (FrameIdx != FrameNum - 1) { Output += TEXT(" "); }
+					}
+					Output += TEXT("]\n");
+				}
+				else
+				{
+					Output += TEXT("[");
+					for (int32 FrameIdx = 0; FrameIdx < RoundedCuttoff / 2; FrameIdx++)
+					{
+						Output += UE::AnimDatabase::FrameAttribute::Private::FrameAttributeFrameToString(FrameAttribute, FrameOffset + FrameIdx);
+						Output += TEXT(" ");
+					}
+
+					Output += TEXT("... ");
+
+					for (int32 FrameIdx = FrameNum - RoundedCuttoff / 2 - 1; FrameIdx < FrameNum; FrameIdx++)
+					{
+						Output += UE::AnimDatabase::FrameAttribute::Private::FrameAttributeFrameToString(FrameAttribute, FrameOffset + FrameIdx);
+						if (FrameIdx != FrameNum - 1) { Output += TEXT(" "); }
+					}
+
+					Output += TEXT("]\n");
+				}
+			}
+
+			Output += TEXT("    ...\n");
+
+			for (int32 RangeIdx = RangeNum - RoundedCuttoff / 2 - 1; RangeIdx < RangeNum; RangeIdx++)
+			{
+				Output += FString::Printf(TEXT("        %i:%i "),
+					FrameAttribute.FrameAttribute->FrameRangeSet.GetEntryRangeStart(EntryIdx, RangeIdx),
+					FrameAttribute.FrameAttribute->FrameRangeSet.GetEntryRangeStop(EntryIdx, RangeIdx));
+
+				const int32 FrameNum = FrameAttribute.FrameAttribute->FrameRangeSet.GetEntryRangeLength(EntryIdx, RangeIdx);
+				const int32 FrameOffset = FrameAttribute.FrameAttribute->FrameRangeSet.GetEntryRangeOffset(EntryIdx, RangeIdx);
+
+				if (FrameNum < RoundedCuttoff)
+				{
+					Output += TEXT("[");
+					for (int32 FrameIdx = 0; FrameIdx < FrameNum; FrameIdx++)
+					{
+						Output += UE::AnimDatabase::FrameAttribute::Private::FrameAttributeFrameToString(FrameAttribute, FrameOffset + FrameIdx);
+						if (FrameIdx != FrameNum - 1) { Output += TEXT(" "); }
+					}
+					Output += TEXT("]\n");
+				}
+				else
+				{
+					Output += TEXT("[");
+					for (int32 FrameIdx = 0; FrameIdx < RoundedCuttoff / 2; FrameIdx++)
+					{
+						Output += UE::AnimDatabase::FrameAttribute::Private::FrameAttributeFrameToString(FrameAttribute, FrameOffset + FrameIdx);
+						Output += TEXT(" ");
+					}
+
+					Output += TEXT("... ");
+
+					for (int32 FrameIdx = FrameNum - RoundedCuttoff / 2 - 1; FrameIdx < FrameNum; FrameIdx++)
+					{
+						Output += UE::AnimDatabase::FrameAttribute::Private::FrameAttributeFrameToString(FrameAttribute, FrameOffset + FrameIdx);
+						if (FrameIdx != FrameNum - 1) { Output += TEXT(" "); }
+					}
+
+					Output += TEXT("]\n");
+				}
+			}
+		}
+		Output += TEXT("    ]\n");
+	}
+	Output += TEXT("]");
+
+	return Output;
+}
+
 FAnimDatabaseFrameRanges UAnimDatabaseFrameAttributeLibrary::FrameAttributeOutlierFrameRanges(
 	const UAnimDatabase* Database,
 	const FAnimDatabaseFrameRanges& FrameRanges,
diff --git a/Engine/Plugins/Experimental/Animation/AnimDatabase/Source/AnimDatabase/Private/AnimDatabaseFrameRanges.cpp b/Engine/Plugins/Experimental/Animation/AnimDatabase/Source/AnimDatabase/Private/AnimDatabaseFrameRanges.cpp
--- a/Engine/Plugins/Experimental/Animation/AnimDatabase/Source/AnimDatabase/Private/AnimDatabaseFrameRanges.cpp
+++ b/Engine/Plugins/Experimental/Animation/AnimDatabase/Source/AnimDatabase/Private/AnimDatabaseFrameRanges.cpp
@@ -696,6 +696,70 @@ FAnimDatabaseFrames UAnimDatabaseFramesLibrary::FramesDifferenceFromArray(const
 	return FramesDifferenceFromArrayView(Frames);
 }
 
+FString UAnimDatabaseFramesLibrary::FramesToString(const FAnimDatabaseFrames& Frames)
+{
+	if (!Frames.IsValid()) { return TEXT("Invalid"); }
+
+	return FString::Printf(TEXT("EntryNum=%i TotalFrameNum=%i"), Frames.FrameSet->GetEntryNum(), Frames.FrameSet->GetTotalFrameNum());
+}
+
+FString UAnimDatabaseFramesLibrary::FramesToStringFormat(const UAnimDatabase* Database, const FAnimDatabaseFrames& Frames, const int32 Cutoff)
+{
+	if (!Frames.IsValid()) { return TEXT("Invalid"); }
+
+	const int32 RoundedCuttoff = (Cutoff / 2) * 2 + 1;
+
+	FString Output = FString::Printf(TEXT("EntryNum=%i TotalFrameNum=%i Data=[\n"), Frames.FrameSet->GetEntryNum(), Frames.FrameSet->GetTotalFrameNum());
+	const int32 EntryNum = Frames.FrameSet->GetEntryNum();
+	for (int32 EntryIdx = 0; EntryIdx < EntryNum; EntryIdx++)
+	{
+		const int32 SequenceIdx = Frames.FrameSet->GetEntrySequence(EntryIdx);
+		const int32 FrameNum = Frames.FrameSet->GetEntryFrameNum(EntryIdx);
+		
+		if (Database)
+		{
+			Output += FString::Printf(TEXT("   %s%s ["),
+				*Database->GetSequenceAssetName(SequenceIdx),
+				Database->GetIsMirrored(SequenceIdx) ? TEXT(" (mirrored)") : TEXT(""));
+		}
+		else
+		{
+			Output += FString::Printf(TEXT("   Sequence %i ["), SequenceIdx);
+
+		}
+
+		if (Cutoff < 0 || FrameNum < RoundedCuttoff)
+		{
+			for (int32 FrameIdx = 0; FrameIdx < FrameNum; FrameIdx++)
+			{
+				Output += FString::FromInt(Frames.FrameSet->GetEntryFrame(EntryIdx, FrameIdx));
+				if (FrameIdx != FrameNum - 1) { Output += TEXT(" "); }
+			}
+		}
+		else
+		{
+			for (int32 FrameIdx = 0; FrameIdx < RoundedCuttoff / 2; FrameIdx++)
+			{
+				Output += FString::FromInt(Frames.FrameSet->GetEntryFrame(EntryIdx, FrameIdx));
+				Output += TEXT(" ");
+			}
+
+			Output += TEXT("... ");
+
+			for (int32 FrameIdx = FrameNum - RoundedCuttoff / 2 - 1; FrameIdx < FrameNum; FrameIdx++)
+			{
+				Output += FString::FromInt(Frames.FrameSet->GetEntryFrame(EntryIdx, FrameIdx));
+				if (FrameIdx != FrameNum - 1) { Output += TEXT(" "); }
+			}
+		}
+
+		Output += TEXT("]\n");
+	}
+	Output += TEXT("]");
+
+	return Output;
+}
+
 FAnimDatabaseFrames UAnimDatabaseFramesFunction::MakeFrames_Implementation(const UAnimDatabase* Database, const FAnimDatabaseFrameRanges& FrameRanges) const
 {
 	return UAnimDatabaseFramesLibrary::MakeEmptyFrames();
@@ -2224,6 +2288,79 @@ FAnimDatabaseFrameRanges UAnimDatabaseFrameRangesLibrary::FrameRangesRangeAtInde
 	return Out;
 }
 
+FString UAnimDatabaseFrameRangesLibrary::FrameRangesToString(const FAnimDatabaseFrameRanges& FrameRanges)
+{
+	if (!FrameRanges.IsValid()) { return TEXT("Invalid"); }
+
+	return FString::Printf(TEXT("EntryNum=%i TotalRangeNum=%i TotalFrameNum=%i"), 
+		FrameRanges.FrameRangeSet->GetEntryNum(), FrameRanges.FrameRangeSet->GetTotalRangeNum(), FrameRanges.FrameRangeSet->GetTotalFrameNum());
+}
+
+FString UAnimDatabaseFrameRangesLibrary::FrameRangesToStringFormat(const UAnimDatabase* Database, const FAnimDatabaseFrameRanges& FrameRanges, const int32 Cutoff)
+{
+	if (!FrameRanges.IsValid()) { return TEXT("Invalid"); }
+
+	const int32 RoundedCuttoff = (Cutoff / 2) * 2 + 1;
+
+	FString Output = FString::Printf(TEXT("EntryNum=%i TotalRangeNum=%i TotalFrameNum=%i Data=[\n"),
+		FrameRanges.FrameRangeSet->GetEntryNum(), FrameRanges.FrameRangeSet->GetTotalRangeNum(), FrameRanges.FrameRangeSet->GetTotalFrameNum());
+
+	const int32 EntryNum = FrameRanges.FrameRangeSet->GetEntryNum();
+	for (int32 EntryIdx = 0; EntryIdx < EntryNum; EntryIdx++)
+	{
+		const int32 SequenceIdx = FrameRanges.FrameRangeSet->GetEntrySequence(EntryIdx);
+		const int32 RangeNum = FrameRanges.FrameRangeSet->GetEntryRangeNum(EntryIdx);
+
+		if (Database)
+		{
+			Output += FString::Printf(TEXT("   %s%s ["),
+				*Database->GetSequenceAssetName(SequenceIdx),
+				Database->GetIsMirrored(SequenceIdx) ? TEXT(" (mirrored)") : TEXT(""));
+		}
+		else
+		{
+			Output += FString::Printf(TEXT("   Sequence %i ["), SequenceIdx);
+
+		}
+		
+		if (Cutoff < 0 || RangeNum < RoundedCuttoff)
+		{
+			for (int32 RangeIdx = 0; RangeIdx < RangeNum; RangeIdx++)
+			{
+				Output += FString::Printf(TEXT("%i:%i"),
+					FrameRanges.FrameRangeSet->GetEntryRangeStart(EntryIdx, RangeIdx),
+					FrameRanges.FrameRangeSet->GetEntryRangeStop(EntryIdx, RangeIdx));
+				if (RangeIdx != RangeNum - 1) { Output += TEXT(" "); }
+			}
+		}
+		else
+		{
+			for (int32 RangeIdx = 0; RangeIdx < RoundedCuttoff / 2; RangeIdx++)
+			{
+				Output += FString::Printf(TEXT("%i:%i"),
+					FrameRanges.FrameRangeSet->GetEntryRangeStart(EntryIdx, RangeIdx),
+					FrameRanges.FrameRangeSet->GetEntryRangeStop(EntryIdx, RangeIdx));
+				Output += TEXT(" ");
+			}
+
+			Output += TEXT("... ");
+
+			for (int32 RangeIdx = RangeNum - RoundedCuttoff / 2 - 1; RangeIdx < RangeNum; RangeIdx++)
+			{
+				Output += FString::Printf(TEXT("%i:%i"),
+					FrameRanges.FrameRangeSet->GetEntryRangeStart(EntryIdx, RangeIdx),
+					FrameRanges.FrameRangeSet->GetEntryRangeStop(EntryIdx, RangeIdx));
+				if (RangeIdx != RangeNum - 1) { Output += TEXT(" "); }
+			}
+		}
+
+		Output += TEXT("]\n");
+	}
+	Output += TEXT("]");
+
+	return Output;
+}
+
 void UAnimDatabaseFrameRangesLibrary::FindAnimNotifyClassesInFrameRanges(TArray<TSubclassOf<UAnimNotify>>& OutAnimNotifyClasses, const UAnimDatabase* Database, const FAnimDatabaseFrameRanges& FrameRanges)
 {
 	if (!Database)
diff --git a/Engine/Plugins/Experimental/Animation/AnimDatabase/Source/AnimDatabase/Private/AnimDatabaseMath.cpp b/Engine/Plugins/Experimental/Animation/AnimDatabase/Source/AnimDatabase/Private/AnimDatabaseMath.cpp
--- a/Engine/Plugins/Experimental/Animation/AnimDatabase/Source/AnimDatabase/Private/AnimDatabaseMath.cpp
+++ b/Engine/Plugins/Experimental/Animation/AnimDatabase/Source/AnimDatabase/Private/AnimDatabaseMath.cpp
@@ -240,12 +240,12 @@ namespace UE::AnimDatabase::Math
 
 	float HalfLifeToDamping(const float HalfLife)
 	{
-		return (4.0f * UE_LN2) / (HalfLife + UE_SMALL_NUMBER);
+		return (4.0f * UE_LN2) / FMath::Max(HalfLife, UE_SMALL_NUMBER);
 	}
 
 	float DampingToHalfLife(const float Damping)
 	{
-		return (4.0 * UE_LN2) / (Damping + UE_SMALL_NUMBER);
+		return (4.0 * UE_LN2) / FMath::Max(Damping, UE_SMALL_NUMBER);
 	}
 
 	void CriticalSpringUpdate(
@@ -500,10 +500,10 @@ namespace UE::AnimDatabase::Math
 	{
 		InOutTimeSinceTransition += DeltaTime;
 
-		const float T = FMath::Clamp(InOutTimeSinceTransition / (BlendTime + UE_SMALL_NUMBER), 0.0f, 1.0f);
+		const float T = FMath::Clamp(InOutTimeSinceTransition / FMath::Max(BlendTime, UE_SMALL_NUMBER), 0.0f, 1.0f);
 		const float W0 = 2.0f * T * T * T - 3.0f * T * T + 1.0f;
 		const float W1 = (T * T * T - 2.0f * T * T + T) * BlendTime;
-		const float W2 = (6.0f * T * T - 6.0f * T) / (BlendTime + UE_SMALL_NUMBER);
+		const float W2 = (6.0f * T * T - 6.0f * T) / FMath::Max(BlendTime, UE_SMALL_NUMBER);
 		const float W3 = 3.0f * T * T - 4.0f * T + 1.0f;
 
 		OutLocation = InLocation + (FVector)(W0 * InLocationOffset + W1 * InLinearVelocityOffset);
@@ -520,10 +520,10 @@ namespace UE::AnimDatabase::Math
 		const FVector3f InDstLinearVelocity,
 		const float BlendTime)
 	{
-		const float T = FMath::Clamp(InOutTimeSinceTransition / (BlendTime + UE_SMALL_NUMBER), 0.0f, 1.0f);
+		const float T = FMath::Clamp(InOutTimeSinceTransition / FMath::Max(BlendTime, UE_SMALL_NUMBER), 0.0f, 1.0f);
 		const float W0 = 2.0f * T * T * T - 3.0f * T * T + 1.0f;
 		const float W1 = (T * T * T - 2.0f * T * T + T) * BlendTime;
-		const float W2 = (6.0f * T * T - 6.0f * T) / (BlendTime + UE_SMALL_NUMBER);
+		const float W2 = (6.0f * T * T - 6.0f * T) / FMath::Max(BlendTime, UE_SMALL_NUMBER);
 		const float W3 = 3.0f * T * T - 4.0f * T + 1.0f;
 
 		const FVector OffsetSrcLocation = InSrcLocation + (FVector)(W0 * InOutLocationOffset + W1 * InOutLinearVelocityOffset);
@@ -545,7 +545,23 @@ namespace UE::AnimDatabase::Math
 		const float Q3 = 3 * X * X - 2 * X;
 
 		OutValue = W1 * (P1 - P0) + W2 * V0 + W3 * V1 + P0;
-		OutValueVelocity = (Q1 * (P1 - P0) + Q2 * V0 + Q3 * V1) / FrameTime;
+		OutValueVelocity = (Q1 * (P1 - P0) + Q2 * V0 + Q3 * V1) / FMath::Max(FrameTime, UE_SMALL_NUMBER);
+	}
+
+	void HermiteAngleInterpolate(float& OutAngle, float& OutAngularVelocity, const float A0, const float A1, const float V0, const float V1, const float X, const float FrameTime)
+	{
+		const float W1 = 3 * X * X - 2 * X * X * X;
+		const float W2 = X * X * X - 2 * X * X + X;
+		const float W3 = X * X * X - X * X;
+
+		const float Q1 = 6 * X - 6 * X * X;
+		const float Q2 = 3 * X * X - 4 * X + 1;
+		const float Q3 = 3 * X * X - 2 * X;
+
+		const float A1SubA0 = FMath::FindDeltaAngleRadians(A0, A1);
+
+		OutAngle = FMath::Wrap(W1 * A1SubA0 + W2 * V0 + W3 * V1 + A0, -UE_PI, UE_PI);
+		OutAngularVelocity = (Q1 * A1SubA0 + Q2 * V0 + Q3 * V1) / FMath::Max(FrameTime, UE_SMALL_NUMBER);
 	}
 
 	void HermiteLocationInterpolate(FVector& OutPosition, FVector3f& OutLinearVelocity, const FVector P0, const FVector P1, const FVector3f V0, const FVector3f V1, const float X, const float FrameTime)
@@ -559,7 +575,7 @@ namespace UE::AnimDatabase::Math
 		const float Q3 = 3 * X * X - 2 * X;
 
 		OutPosition = W1 * (P1 - P0) + W2 * (FVector)V0 + W3 * (FVector)V1 + P0;
-		OutLinearVelocity = (Q1 * (FVector3f)(P1 - P0) + Q2 * V0 + Q3 * V1) / FrameTime;
+		OutLinearVelocity = (Q1 * (FVector3f)(P1 - P0) + Q2 * V0 + Q3 * V1) / FMath::Max(FrameTime, UE_SMALL_NUMBER);
 	}
 
 	void HermiteLocationInterpolate(FVector3f& OutPosition, FVector3f& OutLinearVelocity, const FVector3f P0, const FVector3f P1, const FVector3f V0, const FVector3f V1, const float X, const float FrameTime)
@@ -573,7 +589,7 @@ namespace UE::AnimDatabase::Math
 		const float Q3 = 3 * X * X - 2 * X;
 
 		OutPosition = W1 * (P1 - P0) + W2 * V0 + W3 * V1 + P0;
-		OutLinearVelocity = (Q1 * (P1 - P0) + Q2 * V0 + Q3 * V1) / FrameTime;
+		OutLinearVelocity = (Q1 * (P1 - P0) + Q2 * V0 + Q3 * V1) / FMath::Max(FrameTime, UE_SMALL_NUMBER);
 	}
 
 	void HermiteRotationInterpolate(FQuat4f& OutRotation, FVector3f& OutAngularVelocity, const FQuat4f R0, const FQuat4f R1, const FVector3f V0, const FVector3f V1, const float X, const float FrameTime)
@@ -590,7 +606,7 @@ namespace UE::AnimDatabase::Math
 
 		OutRotation = FQuat4f::MakeFromRotationVector(W1 * R1SubR0 + W2 * V0 + W3 * V1) * R0;
 		check(OutRotation.IsNormalized());
-		OutAngularVelocity = (Q1 * R1SubR0 + Q2 * V0 + Q3 * V1) / FrameTime;
+		OutAngularVelocity = (Q1 * R1SubR0 + Q2 * V0 + Q3 * V1) / FMath::Max(FrameTime, UE_SMALL_NUMBER);
 	}
 
 	void HermiteScaleInterpolate(FVector3f& OutScale, FVector3f& OutScalarVelocity, const FVector3f S0, const FVector3f S1, const FVector3f V0, const FVector3f V1, const float X, const float FrameTime)
@@ -606,7 +622,7 @@ namespace UE::AnimDatabase::Math
 		const FVector3f R1SubR0 = VectorLogSafe(VectorDivMax(S1, S0, UE_SMALL_NUMBER));
 
 		OutScale = VectorExp(VectorMin(W1 * R1SubR0 + W2 * V0 + W3 * V1, 10.0f)) * S0;
-		OutScalarVelocity = (Q1 * R1SubR0 + Q2 * V0 + Q3 * V1) / FrameTime;
+		OutScalarVelocity = (Q1 * R1SubR0 + Q2 * V0 + Q3 * V1) / FMath::Max(FrameTime, UE_SMALL_NUMBER);
 	}
 
 	void HermiteArrayInterpolate(
@@ -632,7 +648,7 @@ namespace UE::AnimDatabase::Math
 		for (int64 ValueIdx = 0; ValueIdx < ValueNum; ValueIdx++)
 		{
 			OutValues[ValueIdx] = W1 * (P1[ValueIdx] - P0[ValueIdx]) + W2 * V0[ValueIdx] + W3 * V1[ValueIdx] + P0[ValueIdx];
-			OutValueVelocities[ValueIdx] = (Q1 * (P1[ValueIdx] - P0[ValueIdx]) + Q2 * V0[ValueIdx] + Q3 * V1[ValueIdx]) / FrameTime;
+			OutValueVelocities[ValueIdx] = (Q1 * (P1[ValueIdx] - P0[ValueIdx]) + Q2 * V0[ValueIdx] + Q3 * V1[ValueIdx]) / FMath::Max(FrameTime, UE_SMALL_NUMBER);
 		}
 	}
 
@@ -645,6 +661,17 @@ namespace UE::AnimDatabase::Math
 		OutValue = W1 * (P1 - P0) + W2 * V0 + W3 * V1 + P0;
 	}
 
+	void HermiteAngleInterpolate(float& OutAngle, const float A0, const float A1, const float V0, const float V1, const float X, const float FrameTime)
+	{
+		const float W1 = 3 * X * X - 2 * X * X * X;
+		const float W2 = X * X * X - 2 * X * X + X;
+		const float W3 = X * X * X - X * X;
+
+		const float A1SubA0 = FMath::FindDeltaAngleRadians(A0, A1);
+
+		OutAngle = FMath::Wrap(W1 * A1SubA0 + W2 * V0 + W3 * V1 + A0, -UE_PI, UE_PI);
+	}
+
 	void HermiteLocationInterpolate(FVector& OutPosition, const FVector P0, const FVector P1, const FVector3f V0, const FVector3f V1, const float X, const float FrameTime)
 	{
 		const float W1 = 3 * X * X - 2 * X * X * X;
@@ -712,26 +739,32 @@ namespace UE::AnimDatabase::Math
 	void ValueInterpolateLinear(float& OutValue, float& OutValueVelocity, const float P0, const float P1, const float Alpha, const float FrameTime)
 	{
 		OutValue = FMath::Lerp(P0, P1, Alpha);
-		OutValueVelocity = (P1 - P0) / FrameTime;
+		OutValueVelocity = (P1 - P0) / FMath::Max(FrameTime, UE_SMALL_NUMBER);
+	}
+
+	void AngleInterpolateLinear(float& OutAngle, float& OutAngularVelocity, const float A0, const float A1, const float Alpha, const float FrameTime)
+	{
+		OutAngle = FMath::Wrap(A0 + FMath::FindDeltaAngleRadians(A0, A1) * Alpha, -UE_PI, UE_PI);
+		OutAngularVelocity = FMath::FindDeltaAngleRadians(A0, A1) / FMath::Max(FrameTime, UE_SMALL_NUMBER);
 	}
 
 	void LocationInterpolateLinear(FVector& OutLocation, FVector3f& OutLinearVelocity, const FVector P0, const FVector P1, const float Alpha, const float FrameTime)
 	{
 		OutLocation = FMath::Lerp(P0, P1, Alpha);
-		OutLinearVelocity = (FVector3f)((P1 - P0) / FrameTime);
+		OutLinearVelocity = (FVector3f)((P1 - P0) / FMath::Max(FrameTime, UE_SMALL_NUMBER));
 	}
 
 	void LocationInterpolateLinear(FVector3f& OutLocation, FVector3f& OutLinearVelocity, const FVector3f P0, const FVector3f P1, const float Alpha, const float FrameTime)
 	{
 		OutLocation = FMath::Lerp(P0, P1, Alpha);
-		OutLinearVelocity = (P1 - P0) / FrameTime;
+		OutLinearVelocity = (P1 - P0) / FMath::Max(FrameTime, UE_SMALL_NUMBER);
 	}
 
 	void RotationInterpolateLinear(FQuat4f& OutRotation, FVector3f& OutAngularVelocity, const FQuat4f R0, const FQuat4f R1, const float Alpha, const float FrameTime)
 	{
 		OutRotation = FQuat4f::Slerp(R0, R1, Alpha);
 		check(OutRotation.IsNormalized());
-		OutAngularVelocity = (R1 * R0.Inverse()).GetShortestArcWith(FQuat4f::Identity).ToRotationVector() / FrameTime;
+		OutAngularVelocity = (R1 * R0.Inverse()).GetShortestArcWith(FQuat4f::Identity).ToRotationVector() / FMath::Max(FrameTime, UE_SMALL_NUMBER);
 	}
 
 	void ScaleInterpolateLinear(FVector3f& OutScale, FVector3f& OutScalarVelocity, const FVector3f S0, const FVector3f S1, const float Alpha, const float FrameTime)
@@ -740,7 +773,7 @@ namespace UE::AnimDatabase::Math
 		const FVector3f S1Log = VectorLogSafe(S1);
 
 		OutScale = VectorExp(FMath::Lerp(S0Log, S1Log, Alpha));
-		OutScalarVelocity = (S1Log - S0Log) / FrameTime;
+		OutScalarVelocity = (S1Log - S0Log) / FMath::Max(FrameTime, UE_SMALL_NUMBER);
 	}
 
 	void ArrayInterpolateLinear(
@@ -756,7 +789,7 @@ namespace UE::AnimDatabase::Math
 		for (int64 ValueIdx = 0; ValueIdx < ValueNum; ValueIdx++)
 		{
 			OutValues[ValueIdx] = FMath::Lerp(P0[ValueIdx], P1[ValueIdx], Alpha);
-			OutValueVelocities[ValueIdx] = (P1[ValueIdx] - P0[ValueIdx]) / FrameTime;
+			OutValueVelocities[ValueIdx] = (P1[ValueIdx] - P0[ValueIdx]) / FMath::Max(FrameTime, UE_SMALL_NUMBER);
 		}
 	}
 
@@ -765,6 +798,11 @@ namespace UE::AnimDatabase::Math
 		OutValue = FMath::Lerp(P0, P1, Alpha);
 	}
 
+	void AngleInterpolateLinear(float& OutAngle, const float A0, const float A1, const float Alpha, const float FrameTime)
+	{
+		OutAngle = FMath::Wrap(A0 + FMath::FindDeltaAngleRadians(A0, A1) * Alpha, -UE_PI, UE_PI);
+	}
+
 	void LocationInterpolateLinear(FVector& OutLocation, const FVector P0, const FVector P1, const float Alpha, const float FrameTime)
 	{
 		OutLocation = FMath::Lerp(P0, P1, Alpha);
@@ -828,6 +866,16 @@ namespace UE::AnimDatabase::Math
 		HermiteValueInterpolate(OutValue, OutValueVelocity, P1, P2, V1, V2, Alpha, FrameTime);
 	}
 
+	void AngleInterpolateCubic(float& OutAngle, float& OutAngularVelocity, const float A0, const float A1, const float A2, const float A3, const float Alpha, const float FrameTime)
+	{
+		const float A1A0Diff = FMath::FindDeltaAngleRadians(A0, A1);
+		const float A2A1Diff = FMath::FindDeltaAngleRadians(A1, A2);
+		const float A3A2Diff = FMath::FindDeltaAngleRadians(A2, A3);
+		const float V1 = (A1A0Diff + A2A1Diff) / 2;
+		const float V2 = (A2A1Diff + A3A2Diff) / 2;
+		HermiteAngleInterpolate(OutAngle, OutAngularVelocity, A1, A2, V1, V2, Alpha, FrameTime);
+	}
+
 	void LocationInterpolateCubic(FVector& OutLocation, FVector3f& OutLinearVelocity, const FVector P0, const FVector P1, const FVector P2, const FVector P3, const float Alpha, const float FrameTime)
 	{
 		const FVector3f V1 = (FVector3f)((P1 - P0) + (P2 - P1)) / 2;
@@ -889,6 +937,16 @@ namespace UE::AnimDatabase::Math
 		HermiteValueInterpolate(OutValue, P1, P2, V1, V2, Alpha, FrameTime);
 	}
 
+	void AngleInterpolateCubic(float& OutAngle, const float A0, const float A1, const float A2, const float A3, const float Alpha, const float FrameTime)
+	{
+		const float A1A0Diff = FMath::FindDeltaAngleRadians(A0, A1);
+		const float A2A1Diff = FMath::FindDeltaAngleRadians(A1, A2);
+		const float A3A2Diff = FMath::FindDeltaAngleRadians(A2, A3);
+		const float V1 = (A1A0Diff + A2A1Diff) / 2;
+		const float V2 = (A2A1Diff + A3A2Diff) / 2;
+		HermiteAngleInterpolate(OutAngle, A1, A2, V1, V2, Alpha, FrameTime);
+	}
+
 	void LocationInterpolateCubic(FVector& OutLocation, const FVector P0, const FVector P1, const FVector P2, const FVector P3, const float Alpha, const float FrameTime)
 	{
 		const FVector3f V1 = (FVector3f)((P1 - P0) + (P2 - P1)) / 2;
@@ -966,6 +1024,15 @@ namespace UE::AnimDatabase::Math
 		HermiteValueInterpolate(OutValue, OutValueVelocity, P1, P2, V1, V2, Alpha, FrameTime);
 	}
 
+	void AngleInterpolateCubicStart(float& OutAngle, float& OutAngularVelocity, const float A1, const float A2, const float A3, const float Alpha, const float FrameTime)
+	{
+		const float A2A1Diff = FMath::FindDeltaAngleRadians(A1, A2);
+		const float A3A2Diff = FMath::FindDeltaAngleRadians(A2, A3);
+		const float V1 = A2A1Diff;
+		const float V2 = (A2A1Diff + A3A2Diff) / 2;
+		HermiteAngleInterpolate(OutAngle, OutAngularVelocity, A1, A2, V1, V2, Alpha, FrameTime);
+	}
+
 	void LocationInterpolateCubicStart(FVector& OutLocation, FVector3f& OutLinearVelocity, const FVector P1, const FVector P2, const FVector P3, const float Alpha, const float FrameTime)
 	{
 		const FVector3f V1 = (FVector3f)(P2 - P1);
@@ -1024,6 +1091,15 @@ namespace UE::AnimDatabase::Math
 		HermiteValueInterpolate(OutValue, P1, P2, V1, V2, Alpha, FrameTime);
 	}
 
+	void AngleInterpolateCubicStart(float& OutAngle, const float A1, const float A2, const float A3, const float Alpha, const float FrameTime)
+	{
+		const float A2A1Diff = FMath::FindDeltaAngleRadians(A1, A2);
+		const float A3A2Diff = FMath::FindDeltaAngleRadians(A2, A3);
+		const float V1 = A2A1Diff;
+		const float V2 = (A2A1Diff + A3A2Diff) / 2;
+		HermiteAngleInterpolate(OutAngle, A1, A2, V1, V2, Alpha, FrameTime);
+	}
+
 	void LocationInterpolateCubicStart(FVector& OutLocation, const FVector P1, const FVector P2, const FVector P3, const float Alpha, const float FrameTime)
 	{
 		const FVector3f V1 = (FVector3f)(P2 - P1);
@@ -1098,6 +1174,15 @@ namespace UE::AnimDatabase::Math
 		HermiteValueInterpolate(OutValue, OutValueVelocity, P1, P2, V1, V2, Alpha, FrameTime);
 	}
 
+	void AngleInterpolateCubicEnd(float& OutAngle, float& OutAngularVelocity, const float A0, const float A1, const float A2, const float Alpha, const float FrameTime)
+	{
+		const float A1A0Diff = FMath::FindDeltaAngleRadians(A0, A1);
+		const float A2A1Diff = FMath::FindDeltaAngleRadians(A1, A2);
+		const float V1 = (A1A0Diff + A2A1Diff) / 2;
+		const float V2 = A2A1Diff;
+		HermiteAngleInterpolate(OutAngle, OutAngularVelocity, A1, A2, V1, V2, Alpha, FrameTime);
+	}
+
 	void LocationInterpolateCubicEnd(FVector& OutLocation, FVector3f& OutLinearVelocity, const FVector P0, const FVector P1, const FVector P2, const float Alpha, const float FrameTime)
 	{
 		const FVector3f V1 = (FVector3f)((P1 - P0) + (P2 - P1)) / 2;
@@ -1156,6 +1241,15 @@ namespace UE::AnimDatabase::Math
 		HermiteValueInterpolate(OutValue, P1, P2, V1, V2, Alpha, FrameTime);
 	}
 
+	void AngleInterpolateCubicEnd(float& OutAngle, const float A0, const float A1, const float A2, const float Alpha, const float FrameTime)
+	{
+		const float A1A0Diff = FMath::FindDeltaAngleRadians(A0, A1);
+		const float A2A1Diff = FMath::FindDeltaAngleRadians(A1, A2);
+		const float V1 = (A1A0Diff + A2A1Diff) / 2;
+		const float V2 = A2A1Diff;
+		HermiteAngleInterpolate(OutAngle, A1, A2, V1, V2, Alpha, FrameTime);
+	}
+
 	void LocationInterpolateCubicEnd(FVector& OutLocation, const FVector P0, const FVector P1, const FVector P2, const float Alpha, const float FrameTime)
 	{
 		const FVector3f V1 = (FVector3f)((P1 - P0) + (P2 - P1)) / 2;
@@ -1249,6 +1343,16 @@ namespace UE::AnimDatabase::Math
 		HermiteValueInterpolate(OutValue, OutValueVelocity, P1, P2, V1, V2, Alpha, FrameTime);
 	}
 
+	void AngleInterpolateCubicMono(float& OutAngle, float& OutAngularVelocity, const float A0, const float A1, const float A2, const float A3, const float Alpha, const float FrameTime)
+	{
+		const float A1A0Diff = FMath::FindDeltaAngleRadians(A0, A1);
+		const float A2A1Diff = FMath::FindDeltaAngleRadians(A1, A2);
+		const float A3A2Diff = FMath::FindDeltaAngleRadians(A2, A3);
+		const float V1 = ComputeMonotoneVelocity(A1A0Diff, A2A1Diff);
+		const float V2 = ComputeMonotoneVelocity(A2A1Diff, A3A2Diff);
+		HermiteAngleInterpolate(OutAngle, OutAngularVelocity, A1, A2, V1, V2, Alpha, FrameTime);
+	}
+
 	void LocationInterpolateCubicMono(FVector& OutLocation, FVector3f& OutLinearVelocity, const FVector P0, const FVector P1, const FVector P2, const FVector P3, const float Alpha, const float FrameTime)
 	{
 		const FVector3f V1 = ComputeMonotoneVelocity((FVector3f)(P1 - P0), (FVector3f)(P2 - P1));
@@ -1310,6 +1414,16 @@ namespace UE::AnimDatabase::Math
 		HermiteValueInterpolate(OutValue, P1, P2, V1, V2, Alpha, FrameTime);
 	}
 
+	void AngleInterpolateCubicMono(float& OutAngle, const float A0, const float A1, const float A2, const float A3, const float Alpha, const float FrameTime)
+	{
+		const float A1A0Diff = FMath::FindDeltaAngleRadians(A0, A1);
+		const float A2A1Diff = FMath::FindDeltaAngleRadians(A1, A2);
+		const float A3A2Diff = FMath::FindDeltaAngleRadians(A2, A3);
+		const float V1 = ComputeMonotoneVelocity(A1A0Diff, A2A1Diff);
+		const float V2 = ComputeMonotoneVelocity(A2A1Diff, A3A2Diff);
+		HermiteAngleInterpolate(OutAngle, A1, A2, V1, V2, Alpha, FrameTime);
+	}
+
 	void LocationInterpolateCubicMono(FVector& OutLocation, const FVector P0, const FVector P1, const FVector P2, const FVector P3, const float Alpha, const float FrameTime)
 	{
 		const FVector3f V1 = ComputeMonotoneVelocity((FVector3f)(P1 - P0), (FVector3f)(P2 - P1));
@@ -1387,6 +1501,15 @@ namespace UE::AnimDatabase::Math
 		HermiteValueInterpolate(OutValue, OutValueVelocity, P1, P2, V1, V2, Alpha, FrameTime);
 	}
 
+	void AngleInterpolateCubicMonoStart(float& OutAngle, float& OutAngularVelocity, const float A1, const float A2, const float A3, const float Alpha, const float FrameTime)
+	{
+		const float A2A1Diff = FMath::FindDeltaAngleRadians(A1, A2);
+		const float A3A2Diff = FMath::FindDeltaAngleRadians(A2, A3);
+		const float V1 = A2A1Diff;
+		const float V2 = ComputeMonotoneVelocity(A2A1Diff, A3A2Diff);
+		HermiteAngleInterpolate(OutAngle, OutAngularVelocity, A1, A2, V1, V2, Alpha, FrameTime);
+	}
+
 	void LocationInterpolateCubicMonoStart(FVector& OutLocation, FVector3f& OutLinearVelocity, const FVector P1, const FVector P2, const FVector P3, const float Alpha, const float FrameTime)
 	{
 		const FVector3f V1 = (FVector3f)(P2 - P1);
@@ -1445,6 +1568,15 @@ namespace UE::AnimDatabase::Math
 		HermiteValueInterpolate(OutValue, P1, P2, V1, V2, Alpha, FrameTime);
 	}
 
+	void AngleInterpolateCubicMonoStart(float& OutAngle, const float A1, const float A2, const float A3, const float Alpha, const float FrameTime)
+	{
+		const float A2A1Diff = FMath::FindDeltaAngleRadians(A1, A2);
+		const float A3A2Diff = FMath::FindDeltaAngleRadians(A2, A3);
+		const float V1 = A2A1Diff;
+		const float V2 = ComputeMonotoneVelocity(A2A1Diff, A3A2Diff);
+		HermiteAngleInterpolate(OutAngle, A1, A2, V1, V2, Alpha, FrameTime);
+	}
+
 	void LocationInterpolateCubicMonoStart(FVector& OutLocation, const FVector P1, const FVector P2, const FVector P3, const float Alpha, const float FrameTime)
 	{
 		const FVector3f V1 = (FVector3f)(P2 - P1);
@@ -1519,6 +1651,15 @@ namespace UE::AnimDatabase::Math
 		HermiteValueInterpolate(OutValue, OutValueVelocity, P1, P2, V1, V2, Alpha, FrameTime);
 	}
 
+	void AngleInterpolateCubicMonoEnd(float& OutAngle, float& OutAngularVelocity, const float A0, const float A1, const float A2, const float Alpha, const float FrameTime)
+	{
+		const float A1A0Diff = FMath::FindDeltaAngleRadians(A0, A1);
+		const float A2A1Diff = FMath::FindDeltaAngleRadians(A1, A2);
+		const float V1 = ComputeMonotoneVelocity(A1A0Diff, A2A1Diff);
+		const float V2 = A2A1Diff;
+		HermiteAngleInterpolate(OutAngle, OutAngularVelocity, A1, A2, V1, V2, Alpha, FrameTime);
+	}
+
 	void LocationInterpolateCubicMonoEnd(FVector& OutLocation, FVector3f& OutLinearVelocity, const FVector P0, const FVector P1, const FVector P2, const float Alpha, const float FrameTime)
 	{
 		const FVector3f V1 = ComputeMonotoneVelocity((FVector3f)(P1 - P0), (FVector3f)(P2 - P1));
@@ -1577,6 +1718,15 @@ namespace UE::AnimDatabase::Math
 		HermiteValueInterpolate(OutValue, P1, P2, V1, V2, Alpha, FrameTime);
 	}
 
+	void AngleInterpolateCubicMonoEnd(float& OutAngle, const float A0, const float A1, const float A2, const float Alpha, const float FrameTime)
+	{
+		const float A1A0Diff = FMath::FindDeltaAngleRadians(A0, A1);
+		const float A2A1Diff = FMath::FindDeltaAngleRadians(A1, A2);
+		const float V1 = ComputeMonotoneVelocity(A1A0Diff, A2A1Diff);
+		const float V2 = A2A1Diff;
+		HermiteAngleInterpolate(OutAngle, A1, A2, V1, V2, Alpha, FrameTime);
+	}
+
 	void LocationInterpolateCubicMonoEnd(FVector& OutLocation, const FVector P0, const FVector P1, const FVector P2, const float Alpha, const float FrameTime)
 	{
 		const FVector3f V1 = ComputeMonotoneVelocity((FVector3f)(P1 - P0), (FVector3f)(P2 - P1));
@@ -1651,7 +1801,7 @@ namespace UE::AnimDatabase::Math
 		const int64 FrameNum)
 	{
 		const float ClampedFrameTime = FMath::Clamp(FrameTime, 0.0f, FMath::Max(FrameNum - 1, 0));
-		return FMath::Clamp(FMath::RoundToInt(ClampedFrameTime), 0, FrameNum - 1);
+		return FMath::Clamp(FMath::RoundToInt(ClampedFrameTime), 0, FMath::Max(FrameNum - 1, 0));
 	}
 
 	void ComputeLinearSampleFramesAndAlpha(
@@ -1662,9 +1812,9 @@ namespace UE::AnimDatabase::Math
 		const int64 FrameNum)
 	{
 		const float ClampedFrameTime = FMath::Clamp(FrameTime, 0.0f, FMath::Max(FrameNum - 1, 0));
-		const int64 SampleFrame = FMath::Clamp(FMath::FloorToInt(ClampedFrameTime), 0, FrameNum - 1);
-		OutSampleFrame0 = FMath::Clamp(SampleFrame + 0, 0, FrameNum - 1);
-		OutSampleFrame1 = FMath::Clamp(SampleFrame + 1, 0, FrameNum - 1);
+		const int64 SampleFrame = FMath::Clamp(FMath::FloorToInt(ClampedFrameTime), 0, FMath::Max(FrameNum - 1, 0));
+		OutSampleFrame0 = FMath::Clamp(SampleFrame + 0, 0, FMath::Max(FrameNum - 1, 0));
+		OutSampleFrame1 = FMath::Clamp(SampleFrame + 1, 0, FMath::Max(FrameNum - 1, 0));
 		OutSampleAlpha = ClampedFrameTime - (float)SampleFrame;
 	}
 
@@ -1678,11 +1828,11 @@ namespace UE::AnimDatabase::Math
 		const int64 FrameNum)
 	{
 		const float ClampedFrameTime = FMath::Clamp(FrameTime, 0.0f, FMath::Max(FrameNum - 1, 0));
-		const int64 SampleFrame = FMath::Clamp(FMath::FloorToInt(ClampedFrameTime), 0, FrameNum - 2);
-		OutSampleFrame0 = FMath::Clamp(SampleFrame - 1, 0, FrameNum - 1);
-		OutSampleFrame1 = FMath::Clamp(SampleFrame + 0, 0, FrameNum - 1);
-		OutSampleFrame2 = FMath::Clamp(SampleFrame + 1, 0, FrameNum - 1);
-		OutSampleFrame3 = FMath::Clamp(SampleFrame + 2, 0, FrameNum - 1);
+		const int64 SampleFrame = FMath::Clamp(FMath::FloorToInt(ClampedFrameTime), 0, FMath::Max(FrameNum - 2, 0));
+		OutSampleFrame0 = FMath::Clamp(SampleFrame - 1, 0, FMath::Max(FrameNum - 1, 0));
+		OutSampleFrame1 = FMath::Clamp(SampleFrame + 0, 0, FMath::Max(FrameNum - 1, 0));
+		OutSampleFrame2 = FMath::Clamp(SampleFrame + 1, 0, FMath::Max(FrameNum - 1, 0));
+		OutSampleFrame3 = FMath::Clamp(SampleFrame + 2, 0, FMath::Max(FrameNum - 1, 0));
 		OutSampleAlpha = ClampedFrameTime - (float)SampleFrame;
 	}
 
@@ -1701,9 +1851,9 @@ namespace UE::AnimDatabase::Math
 		const int64 FrameNum)
 	{
 		const int64 ClampedFrameTime = FMath::Clamp(FrameTime, 0, FMath::Max(FrameNum - 1, 0));
-		const int64 SampleFrame = FMath::Clamp(ClampedFrameTime, 0, FrameNum - 1);
-		OutSampleFrame0 = FMath::Clamp(SampleFrame + 0, 0, FrameNum - 1);
-		OutSampleFrame1 = FMath::Clamp(SampleFrame + 1, 0, FrameNum - 1);
+		const int64 SampleFrame = FMath::Clamp(ClampedFrameTime, 0, FMath::Max(FrameNum - 1, 0));
+		OutSampleFrame0 = FMath::Clamp(SampleFrame + 0, 0, FMath::Max(FrameNum - 1, 0));
+		OutSampleFrame1 = FMath::Clamp(SampleFrame + 1, 0, FMath::Max(FrameNum - 1, 0));
 		OutSampleAlpha = ClampedFrameTime - SampleFrame;
 	}
 
@@ -1717,11 +1867,11 @@ namespace UE::AnimDatabase::Math
 		const int64 FrameNum)
 	{
 		const int64 ClampedFrameTime = FMath::Clamp(FrameTime, 0, FMath::Max(FrameNum - 1, 0));
-		const int64 SampleFrame = FMath::Clamp(ClampedFrameTime, 0, FrameNum - 2);
-		OutSampleFrame0 = FMath::Clamp(SampleFrame - 1, 0, FrameNum - 1);
-		OutSampleFrame1 = FMath::Clamp(SampleFrame + 0, 0, FrameNum - 1);
-		OutSampleFrame2 = FMath::Clamp(SampleFrame + 1, 0, FrameNum - 1);
-		OutSampleFrame3 = FMath::Clamp(SampleFrame + 2, 0, FrameNum - 1);
+		const int64 SampleFrame = FMath::Clamp(ClampedFrameTime, 0, FMath::Max(FrameNum - 2, 0));
+		OutSampleFrame0 = FMath::Clamp(SampleFrame - 1, 0, FMath::Max(FrameNum - 1, 0));
+		OutSampleFrame1 = FMath::Clamp(SampleFrame + 0, 0, FMath::Max(FrameNum - 1, 0));
+		OutSampleFrame2 = FMath::Clamp(SampleFrame + 1, 0, FMath::Max(FrameNum - 1, 0));
+		OutSampleFrame3 = FMath::Clamp(SampleFrame + 2, 0, FMath::Max(FrameNum - 1, 0));
 		OutSampleAlpha = ClampedFrameTime - SampleFrame;
 	}
 
@@ -1735,6 +1885,14 @@ namespace UE::AnimDatabase::Math
 		OutValue = InValues[SampleFrame];
 	}
 
+	void AngleSampleNearest(
+		float& OutAngle,
+		const TLearningArrayView<1, const float> InAngles,
+		const int32 SampleFrame)
+	{
+		OutAngle = InAngles[SampleFrame];
+	}
+
 	void LocationSampleNearest(
 		FVector& OutLocation,
 		const TLearningArrayView<1, const FVector> InLocations,
@@ -1855,6 +2013,17 @@ namespace UE::AnimDatabase::Math
 		Math::ValueInterpolateLinear(OutValue, InValues[SampleFrame0], InValues[SampleFrame1], SampleAlpha, FrameDeltaTime);
 	}
 
+	void AngleSampleLinear(
+		float& OutAngle,
+		const TLearningArrayView<1, const float> InAngles,
+		const int32 SampleFrame0,
+		const int32 SampleFrame1,
+		const float SampleAlpha,
+		const float FrameDeltaTime)
+	{
+		Math::AngleInterpolateLinear(OutAngle, InAngles[SampleFrame0], InAngles[SampleFrame1], SampleAlpha, FrameDeltaTime);
+	}
+
 	void LocationSampleLinear(
 		FVector& OutLocation,
 		const TLearningArrayView<1, const FVector> InLocations,
@@ -2014,6 +2183,34 @@ namespace UE::AnimDatabase::Math
 		}
 	}
 
+	void AngleSampleCubic(
+		float& OutAngle,
+		const TLearningArrayView<1, const float> InAngles,
+		const int32 SampleFrame0,
+		const int32 SampleFrame1,
+		const int32 SampleFrame2,
+		const int32 SampleFrame3,
+		const float SampleAlpha,
+		const float FrameDeltaTime)
+	{
+		if (SampleFrame0 == SampleFrame1 && SampleFrame2 == SampleFrame3)
+		{
+			Math::AngleInterpolateLinear(OutAngle, InAngles[SampleFrame1], InAngles[SampleFrame2], SampleAlpha, FrameDeltaTime);
+		}
+		else if (SampleFrame0 == SampleFrame1)
+		{
+			Math::AngleInterpolateCubicStart(OutAngle, InAngles[SampleFrame1], InAngles[SampleFrame2], InAngles[SampleFrame3], SampleAlpha, FrameDeltaTime);
+		}
+		else if (SampleFrame2 == SampleFrame3)
+		{
+			Math::AngleInterpolateCubicEnd(OutAngle, InAngles[SampleFrame0], InAngles[SampleFrame1], InAngles[SampleFrame2], SampleAlpha, FrameDeltaTime);
+		}
+		else
+		{
+			Math::AngleInterpolateCubic(OutAngle, InAngles[SampleFrame0], InAngles[SampleFrame1], InAngles[SampleFrame2], InAngles[SampleFrame3], SampleAlpha, FrameDeltaTime);
+		}
+	}
+
 	void LocationSampleCubic(
 		FVector& OutLocation,
 		const TLearningArrayView<1, const FVector> InLocations,
@@ -2353,6 +2550,34 @@ namespace UE::AnimDatabase::Math
 		}
 	}
 
+	void AngleSampleCubicMono(
+		float& OutAngle,
+		const TLearningArrayView<1, const float> InAngles,
+		const int32 SampleFrame0,
+		const int32 SampleFrame1,
+		const int32 SampleFrame2,
+		const int32 SampleFrame3,
+		const float SampleAlpha,
+		const float FrameDeltaTime)
+	{
+		if (SampleFrame0 == SampleFrame1 && SampleFrame2 == SampleFrame3)
+		{
+			Math::AngleInterpolateLinear(OutAngle, InAngles[SampleFrame1], InAngles[SampleFrame2], SampleAlpha, FrameDeltaTime);
+		}
+		else if (SampleFrame0 == SampleFrame1)
+		{
+			Math::AngleInterpolateCubicMonoStart(OutAngle, InAngles[SampleFrame1], InAngles[SampleFrame2], InAngles[SampleFrame3], SampleAlpha, FrameDeltaTime);
+		}
+		else if (SampleFrame2 == SampleFrame3)
+		{
+			Math::AngleInterpolateCubicMonoEnd(OutAngle, InAngles[SampleFrame0], InAngles[SampleFrame1], InAngles[SampleFrame2], SampleAlpha, FrameDeltaTime);
+		}
+		else
+		{
+			Math::AngleInterpolateCubicMono(OutAngle, InAngles[SampleFrame0], InAngles[SampleFrame1], InAngles[SampleFrame2], InAngles[SampleFrame3], SampleAlpha, FrameDeltaTime);
+		}
+	}
+
 	void LocationSampleCubicMono(
 		FVector& OutLocation,
 		const TLearningArrayView<1, const FVector> InLocations,
@@ -2670,7 +2895,7 @@ namespace UE::AnimDatabase::Math
 		const float Time,
 		const float FrameDeltaTime)
 	{
-		const int64 SampleFrame = ComputeNearestSampleFrame(Time / FrameDeltaTime, InValues.Num<0>());
+		const int64 SampleFrame = ComputeNearestSampleFrame(Time / FMath::Max(FrameDeltaTime, UE_SMALL_NUMBER), InValues.Num<0>());
 		OutSampledValue = InValues[SampleFrame];
 	}
 
@@ -2680,7 +2905,7 @@ namespace UE::AnimDatabase::Math
 		const float Time,
 		const float FrameDeltaTime)
 	{
-		const int64 SampleFrame = ComputeNearestSampleFrame(Time / FrameDeltaTime, InVectors.Num<0>());
+		const int64 SampleFrame = ComputeNearestSampleFrame(Time / FMath::Max(FrameDeltaTime, UE_SMALL_NUMBER), InVectors.Num<0>());
 		Learning::Array::Copy(OutSampledVector, InVectors[SampleFrame]);
 	}
 
@@ -2698,7 +2923,7 @@ namespace UE::AnimDatabase::Math
 			SampleFrame0,
 			SampleFrame1,
 			SampleAlpha,
-			Time / FrameDeltaTime,
+			Time / FMath::Max(FrameDeltaTime, UE_SMALL_NUMBER),
 			FrameNum);
 
 		ValueInterpolateLinear(
@@ -2723,7 +2948,7 @@ namespace UE::AnimDatabase::Math
 			SampleFrame0,
 			SampleFrame1,
 			SampleAlpha,
-			Time / FrameDeltaTime,
+			Time / FMath::Max(FrameDeltaTime, UE_SMALL_NUMBER),
 			FrameNum);
 
 		ArrayInterpolateLinear(
@@ -2750,7 +2975,7 @@ namespace UE::AnimDatabase::Math
 			SampleFrame2,
 			SampleFrame3,
 			SampleAlpha,
-			Time / FrameDeltaTime,
+			Time / FMath::Max(FrameDeltaTime, UE_SMALL_NUMBER),
 			FrameNum);
 
 		if (SampleFrame0 == SampleFrame1 && SampleFrame2 == SampleFrame3)
@@ -2811,7 +3036,7 @@ namespace UE::AnimDatabase::Math
 			SampleFrame2,
 			SampleFrame3,
 			SampleAlpha,
-			Time / FrameDeltaTime,
+			Time / FMath::Max(FrameDeltaTime, UE_SMALL_NUMBER),
 			FrameNum);
 
 		if (SampleFrame0 == SampleFrame1 && SampleFrame2 == SampleFrame3)
@@ -2872,7 +3097,7 @@ namespace UE::AnimDatabase::Math
 			SampleFrame2,
 			SampleFrame3,
 			SampleAlpha,
-			Time / FrameDeltaTime,
+			Time / FMath::Max(FrameDeltaTime, UE_SMALL_NUMBER),
 			FrameNum);
 
 		ValueSampleCubicMono(
@@ -2903,7 +3128,7 @@ namespace UE::AnimDatabase::Math
 			SampleFrame2,
 			SampleFrame3,
 			SampleAlpha,
-			Time / FrameDeltaTime,
+			Time / FMath::Max(FrameDeltaTime, UE_SMALL_NUMBER),
 			FrameNum);
 
 		LocationSampleCubicMono(
@@ -2934,7 +3159,7 @@ namespace UE::AnimDatabase::Math
 			SampleFrame2,
 			SampleFrame3,
 			SampleAlpha,
-			Time / FrameDeltaTime,
+			Time / FMath::Max(FrameDeltaTime, UE_SMALL_NUMBER),
 			FrameNum);
 
 		LocationSampleCubicMono(
@@ -2964,7 +3189,7 @@ namespace UE::AnimDatabase::Math
 			SampleFrame2,
 			SampleFrame3,
 			SampleAlpha,
-			Time / FrameDeltaTime,
+			Time / FMath::Max(FrameDeltaTime, UE_SMALL_NUMBER),
 			FrameNum);
 
 		ArraySampleCubicMono(
diff --git a/Engine/Plugins/Experimental/Animation/AnimDatabase/Source/AnimDatabase/Private/AnimDatabasePose.cpp b/Engine/Plugins/Experimental/Animation/AnimDatabase/Source/AnimDatabase/Private/AnimDatabasePose.cpp
--- a/Engine/Plugins/Experimental/Animation/AnimDatabase/Source/AnimDatabase/Private/AnimDatabasePose.cpp
+++ b/Engine/Plugins/Experimental/Animation/AnimDatabase/Source/AnimDatabase/Private/AnimDatabasePose.cpp
@@ -454,6 +454,14 @@ namespace UE::AnimDatabase
 		return AttributeData[FrameIdx][AttributeOffsets[AttributeIdx]];
 	}
 
+	float FPoseAttributeDataConstView::GetAngle(const int32 FrameIdx, const int32 AttributeIdx) const
+	{
+		check(AttributeTypes[AttributeIdx] == EAnimDatabaseAttributeType::Angle);
+		check(AttributeActive[FrameIdx][AttributeIdx]);
+
+		return AttributeData[FrameIdx][AttributeOffsets[AttributeIdx]];
+	}
+
 	FVector3f FPoseAttributeDataConstView::GetLocation(const int32 FrameIdx, const int32 AttributeIdx) const
 	{
 		check(AttributeTypes[AttributeIdx] == EAnimDatabaseAttributeType::Location);
@@ -568,6 +576,7 @@ namespace UE::AnimDatabase
 	bool FPoseAttributeDataView::GetAttributeActive(const int32 FrameIdx, const int32 AttributeIdx) const { return ConstView().GetAttributeActive(FrameIdx, AttributeIdx); }
 	bool FPoseAttributeDataView::GetBool(const int32 FrameIdx, const int32 AttributeIdx) const { return ConstView().GetBool(FrameIdx, AttributeIdx); }
 	float FPoseAttributeDataView::GetFloat(const int32 FrameIdx, const int32 AttributeIdx) const { return ConstView().GetFloat(FrameIdx, AttributeIdx); }
+	float FPoseAttributeDataView::GetAngle(const int32 FrameIdx, const int32 AttributeIdx) const { return ConstView().GetAngle(FrameIdx, AttributeIdx); }
 	FVector3f FPoseAttributeDataView::GetLocation(const int32 FrameIdx, const int32 AttributeIdx) const { return ConstView().GetLocation(FrameIdx, AttributeIdx); }
 	FQuat4f FPoseAttributeDataView::GetRotation(const int32 FrameIdx, const int32 AttributeIdx) const { return ConstView().GetRotation(FrameIdx, AttributeIdx); }
 	FVector3f FPoseAttributeDataView::GetScale(const int32 FrameIdx, const int32 AttributeIdx) const { return ConstView().GetScale(FrameIdx, AttributeIdx); }
@@ -669,6 +678,7 @@ namespace UE::AnimDatabase
 	bool FPoseAttributeData::GetAttributeActive(const int32 FrameIdx, const int32 AttributeIdx) const { return ConstView().GetAttributeActive(FrameIdx, AttributeIdx); }
 	bool FPoseAttributeData::GetBool(const int32 FrameIdx, const int32 AttributeIdx) const { return ConstView().GetBool(FrameIdx, AttributeIdx); }
 	float FPoseAttributeData::GetFloat(const int32 FrameIdx, const int32 AttributeIdx) const { return ConstView().GetFloat(FrameIdx, AttributeIdx); }
+	float FPoseAttributeData::GetAngle(const int32 FrameIdx, const int32 AttributeIdx) const { return ConstView().GetAngle(FrameIdx, AttributeIdx); }
 	FVector3f FPoseAttributeData::GetLocation(const int32 FrameIdx, const int32 AttributeIdx) const { return ConstView().GetLocation(FrameIdx, AttributeIdx); }
 	FQuat4f FPoseAttributeData::GetRotation(const int32 FrameIdx, const int32 AttributeIdx) const { return ConstView().GetRotation(FrameIdx, AttributeIdx); }
 	FVector3f FPoseAttributeData::GetScale(const int32 FrameIdx, const int32 AttributeIdx) const { return ConstView().GetScale(FrameIdx, AttributeIdx); }
@@ -1034,6 +1044,7 @@ namespace UE::AnimDatabase
 			case EAnimDatabaseAttributeType::Direction: return 1 + 3;
 			case EAnimDatabaseAttributeType::Transform: return 1 + 3 + 6 + 3;
 			case EAnimDatabaseAttributeType::Event: return 1 + 2;
+			case EAnimDatabaseAttributeType::Angle: return 1 + 2;
 			default: checkNoEntry(); return 0;
 			}
 		}
@@ -1404,6 +1415,13 @@ namespace UE::AnimDatabase
 							break;
 						}
 
+						case EAnimDatabaseAttributeType::Angle:
+						{
+							OutPoseVectors[FrameIdx][Offset + 1] = FMath::Sin(InPoseData.AttributeData.AttributeData[FrameIdx][AttributeOffet + 0]);
+							OutPoseVectors[FrameIdx][Offset + 2] = FMath::Cos(InPoseData.AttributeData.AttributeData[FrameIdx][AttributeOffet + 0]);
+							break;
+						}
+
 						default:
 						{
 							checkNoEntry();
@@ -1771,6 +1789,12 @@ namespace UE::AnimDatabase
 							break;
 						}
 
+						case EAnimDatabaseAttributeType::Angle:
+						{
+							OutPoseData.AttributeData.AttributeData[FrameIdx][AttributeOffset + 0] = FMath::Atan2(InPoseVectors[FrameIdx][Offset + 1], InPoseVectors[FrameIdx][Offset + 2]);
+							break;
+						}
+
 						default:
 						{
 							checkNoEntry();
@@ -1897,6 +1921,7 @@ namespace UE::AnimDatabase
 				}
 
 				case EAnimDatabaseAttributeType::Float:
+				case EAnimDatabaseAttributeType::Angle:
 				case EAnimDatabaseAttributeType::Location:
 				case EAnimDatabaseAttributeType::LinearVelocity:
 				case EAnimDatabaseAttributeType::AngularVelocity:
@@ -2757,6 +2782,28 @@ float UAnimDatabasePoseStateLibrary::PoseStateFloatAttribute(const FAnimDatabase
 	return PoseState.CurrPoseData->AttributeData.GetFloat(PoseState.PoseIdx, AttributeIndex);
 }
 
+float UAnimDatabasePoseStateLibrary::PoseStateAngleAttributeRadians(const FAnimDatabasePoseState& PoseState, const int32 AttributeIndex)
+{
+	if (!PoseState.CurrPoseData) { return 0.0f; }
+	if (PoseState.CurrPoseData->IsEmpty()) { return 0.0f; }
+	if (!(AttributeIndex >= 0 && AttributeIndex < PoseState.CurrPoseData->GetAttributeNum())) { return 0.0f; }
+	if (PoseState.CurrPoseData->AttributeData.GetAttributeType(AttributeIndex) != EAnimDatabaseAttributeType::Angle) { return 0.0f; }
+	if (!PoseState.CurrPoseData->AttributeData.GetAttributeActive(PoseState.PoseIdx, AttributeIndex)) { return 0.0f; }
+
+	return PoseState.CurrPoseData->AttributeData.GetAngle(PoseState.PoseIdx, AttributeIndex);
+}
+
+float UAnimDatabasePoseStateLibrary::PoseStateAngleAttributeDegrees(const FAnimDatabasePoseState& PoseState, const int32 AttributeIndex)
+{
+	if (!PoseState.CurrPoseData) { return 0.0f; }
+	if (PoseState.CurrPoseData->IsEmpty()) { return 0.0f; }
+	if (!(AttributeIndex >= 0 && AttributeIndex < PoseState.CurrPoseData->GetAttributeNum())) { return 0.0f; }
+	if (PoseState.CurrPoseData->AttributeData.GetAttributeType(AttributeIndex) != EAnimDatabaseAttributeType::Angle) { return 0.0f; }
+	if (!PoseState.CurrPoseData->AttributeData.GetAttributeActive(PoseState.PoseIdx, AttributeIndex)) { return 0.0f; }
+
+	return FMath::RadiansToDegrees(PoseState.CurrPoseData->AttributeData.GetAngle(PoseState.PoseIdx, AttributeIndex));
+}
+
 FVector UAnimDatabasePoseStateLibrary::PoseStateLocationAttribute(const FAnimDatabasePoseState& PoseState, const int32 AttributeIndex)
 {
 	if (!PoseState.CurrPoseData) { return FVector::ZeroVector; }
diff --git a/Engine/Plugins/Experimental/Animation/AnimDatabase/Source/AnimDatabase/Public/AnimDatabaseFrameAttribute.h b/Engine/Plugins/Experimental/Animation/AnimDatabase/Source/AnimDatabase/Public/AnimDatabaseFrameAttribute.h
--- a/Engine/Plugins/Experimental/Animation/AnimDatabase/Source/AnimDatabase/Public/AnimDatabaseFrameAttribute.h
+++ b/Engine/Plugins/Experimental/Animation/AnimDatabase/Source/AnimDatabase/Public/AnimDatabaseFrameAttribute.h
@@ -38,6 +38,7 @@ enum class EAnimDatabaseAttributeType : uint8
 	Direction = 9,
 	Transform = 10,
 	Event = 11,
+	Angle = 12,
 };
 
 /**
@@ -105,6 +106,12 @@ public:
 	/** Returns the attribute data at the given frame as a double-precision transform. Assumes the Attribute Type is correct. */
 	UE_API FTransform GetAsTransformDouble(const int32 FrameIdx) const;
 
+	/** Returns the attribute data at the given frame as an angle in degrees. Assumes the Attribute Type is correct. */
+	UE_API float GetAsAngleDegrees(const int32 FrameIdx) const;
+
+	/** Returns the attribute data at the given frame as an angle in radians. Assumes the Attribute Type is correct. */
+	UE_API float GetAsAngleRadians(const int32 FrameIdx) const;
+
 	/** Custom Serialization */
 	UE_API bool Serialize(FArchive& Ar);
 
@@ -141,6 +148,11 @@ public:
 	UFUNCTION(BlueprintPure = true, Category = "AnimDatabase")
 	static UE_API int32 AttributeTypeSize(const EAnimDatabaseAttributeType Type);
 
+	/** Gets the name for a given type */
+	UFUNCTION(BlueprintPure = true, Category = "AnimDatabase")
+	static UE_API FString AttributeTypeName(const EAnimDatabaseAttributeType Type);
+	static UE_API const TCHAR* AttributeTypeNameInternal(const EAnimDatabaseAttributeType Type);
+
 	/** Gets the Frame Ranges associated with a given Frame Attribute */
 	UFUNCTION(BlueprintPure, Category = "AnimDatabase", meta = (BlueprintAutocast, CompactNodeTitle = "->"))
 	static UE_API FAnimDatabaseFrameRanges FrameAttributeFrameRanges(const FAnimDatabaseFrameAttribute& FrameAttribute);
@@ -307,6 +319,25 @@ public:
 	UFUNCTION(BlueprintPure = true, Category = "AnimDatabase")
 	static UE_API FAnimDatabaseFrameAttribute MakeScaleFrameAttributeFromConstant(const FAnimDatabaseFrameRanges& FrameRanges, const FVector Scale = FVector(1, 1, 1));
 
+public:
+
+	/** Makes an angle frame attribute from the given ranges, filled with the provided value in degrees. */
+	UFUNCTION(BlueprintPure = true, Category = "AnimDatabase")
+	static UE_API FAnimDatabaseFrameAttribute MakeAngleFrameAttributeFromConstantDegrees(const FAnimDatabaseFrameRanges& FrameRanges, const float Value = 0.0f);
+
+	/** Makes an angle frame attribute from the given ranges, filled with the provided value in radians. */
+	UFUNCTION(BlueprintPure = true, Category = "AnimDatabase")
+	static UE_API FAnimDatabaseFrameAttribute MakeAngleFrameAttributeFromConstantRadians(const FAnimDatabaseFrameRanges& FrameRanges, const float Value = 0.0f);
+
+	/** Makes an angle frame attribute from a float frame attribute with values in degrees */
+	UFUNCTION(BlueprintPure = true, Category = "AnimDatabase")
+	static UE_API FAnimDatabaseFrameAttribute MakeAngleFrameAttributeFromFloatDegrees(const FAnimDatabaseFrameAttribute& A);
+
+	/** Makes an angle frame attribute from a float frame attribute with values in radians */
+	UFUNCTION(BlueprintPure = true, Category = "AnimDatabase")
+	static UE_API FAnimDatabaseFrameAttribute MakeAngleFrameAttributeFromFloatRadians(const FAnimDatabaseFrameAttribute& A);
+
+
 public:
 
 	/** Makes a location frame attribute from the root location of the character at the given frame ranges */
@@ -898,6 +929,14 @@ public:
 	static UE_API bool FrameAttributeLinearVelocitiesAtFrame(TArray<FVector>& OutLinearVelocities, const TArray<FAnimDatabaseFrameAttribute>& FrameAttributes, const int32 SequenceIdx, const int32 FrameIdx);
 	static UE_API bool FrameAttributeLinearVelocitiesAtFrameToArrayView(const TArrayView<FVector> OutLinearVelocities, const TArrayView<const FAnimDatabaseFrameAttribute> FrameAttributes, const int32 SequenceIdx, const int32 FrameIdx);
 
+	/** Sample the given FrameAttribute at the provide Sequence and Frame as an angle in degrees. Returns false if sequence and frame is not within Frame Attribute Ranges. */
+	UFUNCTION(BlueprintPure = true, Category = "AnimDatabase")
+	static UE_API bool FrameAttributeAngleAtFrameDegrees(float& OutValue, const FAnimDatabaseFrameAttribute& FrameAttribute, const int32 SequenceIdx, const int32 FrameIdx);
+
+	/** Sample the given FrameAttribute at the provide Sequence and Frame as an angle in radians. Returns false if sequence and frame is not within Frame Attribute Ranges. */
+	UFUNCTION(BlueprintPure = true, Category = "AnimDatabase")
+	static UE_API bool FrameAttributeAngleAtFrameRadians(float& OutValue, const FAnimDatabaseFrameAttribute& FrameAttribute, const int32 SequenceIdx, const int32 FrameIdx);
+
 
 public:
 
@@ -974,6 +1013,14 @@ public:
 	static UE_API bool FrameAttributeSampleLinearVelocities(TArray<FVector>& OutLinearVelocities, const TArray<FAnimDatabaseFrameAttribute>& FrameAttributes, const int32 SequenceIdx, const float SequenceTime, const FFrameRate& FrameRate);
 	static UE_API bool FrameAttributeSampleLinearVelocitiesToArrayView(const TArrayView<FVector> OutLinearVelocities, const TArrayView<const FAnimDatabaseFrameAttribute> FrameAttributes, const int32 SequenceIdx, const float SequenceTime, const FFrameRate& FrameRate);
 
+	/** Sample the given FrameAttribute at the provide Sequence and Time as an angle in degrees. Returns false if sequence and time is not within Frame Attribute Ranges. */
+	UFUNCTION(BlueprintPure = true, Category = "AnimDatabase")
+	static UE_API bool FrameAttributeSampleAngleDegrees(float& OutValue, const FAnimDatabaseFrameAttribute& FrameAttribute, const int32 SequenceIdx, const float SequenceTime, const FFrameRate& FrameRate);
+
+	/** Sample the given FrameAttribute at the provide Sequence and Time as an angle in radians. Returns false if sequence and time is not within Frame Attribute Ranges. */
+	UFUNCTION(BlueprintPure = true, Category = "AnimDatabase")
+	static UE_API bool FrameAttributeSampleAngleRadians(float& OutValue, const FAnimDatabaseFrameAttribute& FrameAttribute, const int32 SequenceIdx, const float SequenceTime, const FFrameRate& FrameRate);
+
 public:
 
 	/* Gets the total number of ranges in a frame attribute */
@@ -1014,6 +1061,16 @@ public:
 	static UE_API void FrameAttributeToTransformArray(TArray<FTransform>& OutArray, const FAnimDatabaseFrameAttribute& FrameAttribute);
 	static UE_API void FrameAttributeToTransformArrayView(TArray<FTransform>& OutArray, const FAnimDatabaseFrameAttribute& FrameAttribute);
 
+	/** Converts a frame attribute to an array of floats in degrees */
+	UFUNCTION(BlueprintPure = true, Category = "AnimDatabase")
+	static UE_API void FrameAttributeToAngleArrayDegrees(TArray<float>& OutArray, const FAnimDatabaseFrameAttribute& FrameAttribute);
+	static UE_API void FrameAttributeToAngleArrayViewDegrees(const TArrayView<float> OutArray, const FAnimDatabaseFrameAttribute& FrameAttribute);
+
+	/** Converts a frame attribute to an array of floats in radians */
+	UFUNCTION(BlueprintPure = true, Category = "AnimDatabase")
+	static UE_API void FrameAttributeToAngleArrayRadians(TArray<float>& OutArray, const FAnimDatabaseFrameAttribute& FrameAttribute);
+	static UE_API void FrameAttributeToAngleArrayViewRadians(const TArrayView<float> OutArray, const FAnimDatabaseFrameAttribute& FrameAttribute);
+
 public:
 
 	/** Computes a float frame attribute representing the matching distance between some location frame attribute and a constant location value. */
@@ -1229,6 +1286,24 @@ public:
 	static UE_API FAnimDatabaseFrameAttribute FrameAttributeTransformApply(const FTransform& A, const FAnimDatabaseFrameAttribute& B);
 
 
+	/** Adds a constant angle value in degrees to an angle frame attribute. */
+	UFUNCTION(BlueprintPure, Category = "AnimDatabase", meta = (CompactNodeTitle = "+"))
+	static UE_API FAnimDatabaseFrameAttribute FrameAttributeAddAngleDegrees(const FAnimDatabaseFrameAttribute& A, const float B);
+
+	/** Subtracts a constant angle value in degrees to an angle frame attribute. */
+	UFUNCTION(BlueprintPure, Category = "AnimDatabase", meta = (CompactNodeTitle = "-"))
+	static UE_API FAnimDatabaseFrameAttribute FrameAttributeSubtractAngleDegrees(const FAnimDatabaseFrameAttribute& A, const float B);
+
+	/** Adds a constant angle value in radians to an angle frame attribute. */
+	UFUNCTION(BlueprintPure, Category = "AnimDatabase", meta = (CompactNodeTitle = "+"))
+	static UE_API FAnimDatabaseFrameAttribute FrameAttributeAddAngleRadians(const FAnimDatabaseFrameAttribute& A, const float B);
+
+	/** Subtracts a constant angle value in radians to an angle frame attribute. */
+	UFUNCTION(BlueprintPure, Category = "AnimDatabase", meta = (CompactNodeTitle = "-"))
+	static UE_API FAnimDatabaseFrameAttribute FrameAttributeSubtractAngleRadians(const FAnimDatabaseFrameAttribute& A, const float B);
+
+
+
 	/** Adds a frame attribute to multiple others. Assumes types are compatible */
 	UFUNCTION(BlueprintPure, Category = "AnimDatabase")
 	static UE_API void FrameAttributesAdd(TArray<FAnimDatabaseFrameAttribute>& Out, const TArray<FAnimDatabaseFrameAttribute>& A, const FAnimDatabaseFrameAttribute& B);
@@ -1261,6 +1336,10 @@ public:
 
 public:
 
+	/** Copies a frame attribute. */
+	UFUNCTION(BlueprintPure, Category = "AnimDatabase")
+	static UE_API FAnimDatabaseFrameAttribute FrameAttributeCopy(const FAnimDatabaseFrameAttribute& A);
+
 	/** Negates a frame attribute. Assumes types can be negated */
 	UFUNCTION(BlueprintPure, Category = "AnimDatabase")
 	static UE_API FAnimDatabaseFrameAttribute FrameAttributeNegate(const FAnimDatabaseFrameAttribute& A);
@@ -1406,6 +1485,18 @@ public:
 	UFUNCTION(BlueprintPure, Category = "AnimDatabase")
 	static UE_API FAnimDatabaseFrameAttribute FrameAttributeProject(const FAnimDatabaseFrameAttribute& A, const FVector Projection = FVector(1.0f, 1.0f, 0.0f));
 
+	/** Extracts the Yaw component of a rotation frame attribute. */
+	UFUNCTION(BlueprintPure, Category = "AnimDatabase")
+	static UE_API FAnimDatabaseFrameAttribute FrameAttributeYaw(const FAnimDatabaseFrameAttribute& A);
+
+	/** Extracts the Pitch component of a rotation frame attribute. */
+	UFUNCTION(BlueprintPure, Category = "AnimDatabase")
+	static UE_API FAnimDatabaseFrameAttribute FrameAttributePitch(const FAnimDatabaseFrameAttribute& A);
+
+	/** Extracts the Roll component of a rotation frame attribute. */
+	UFUNCTION(BlueprintPure, Category = "AnimDatabase")
+	static UE_API FAnimDatabaseFrameAttribute FrameAttributeRoll(const FAnimDatabaseFrameAttribute& A);
+
 public:
 
 	/** Returns a bool frame attribute describing if two frame attributes are numerically equal on each frame */
@@ -1424,6 +1515,22 @@ public:
 	UFUNCTION(BlueprintPure, Category = "AnimDatabase", meta = (CompactNodeTitle = "<"))
 	static UE_API FAnimDatabaseFrameAttribute FrameAttributeLessThanFloat(const FAnimDatabaseFrameAttribute& FloatFrameAttribute, const float Threshold);
 
+	/** Compares if a angle frame attribute is greater than the given threshold in degrees and returns a bool frame attribute */
+	UFUNCTION(BlueprintPure, Category = "AnimDatabase", meta = (CompactNodeTitle = ">"))
+	static UE_API FAnimDatabaseFrameAttribute FrameAttributeGreaterThanAngleDegrees(const FAnimDatabaseFrameAttribute& AngleFrameAttribute, const float Threshold);
+
+	/** Compares if a angle frame attribute is less than the given threshold in degrees and returns a bool frame attribute */
+	UFUNCTION(BlueprintPure, Category = "AnimDatabase", meta = (CompactNodeTitle = "<"))
+	static UE_API FAnimDatabaseFrameAttribute FrameAttributeLessThanAngleDegrees(const FAnimDatabaseFrameAttribute& AngleFrameAttribute, const float Threshold);
+
+	/** Compares if a angle frame attribute is greater than the given threshold in radians and returns a bool frame attribute */
+	UFUNCTION(BlueprintPure, Category = "AnimDatabase", meta = (CompactNodeTitle = ">"))
+	static UE_API FAnimDatabaseFrameAttribute FrameAttributeGreaterThanAngleRadians(const FAnimDatabaseFrameAttribute& AngleFrameAttribute, const float Threshold);
+
+	/** Compares if a angle frame attribute is less than the given threshold in radians and returns a bool frame attribute */
+	UFUNCTION(BlueprintPure, Category = "AnimDatabase", meta = (CompactNodeTitle = "<"))
+	static UE_API FAnimDatabaseFrameAttribute FrameAttributeLessThanAngleRadians(const FAnimDatabaseFrameAttribute& AngleFrameAttribute, const float Threshold);
+
 	/** Computes the logical not of a bool frame attribute */
 	UFUNCTION(BlueprintPure, Category = "AnimDatabase", meta = (CompactNodeTitle = "NOT"))
 	static UE_API FAnimDatabaseFrameAttribute FrameAttributeNot(const FAnimDatabaseFrameAttribute& A);
@@ -1440,10 +1547,26 @@ public:
 	UFUNCTION(BlueprintPure, Category = "AnimDatabase")
 	static UE_API FAnimDatabaseFrameAttribute FrameAttributeBoolToFloat(const FAnimDatabaseFrameAttribute& BoolFrameAttribute);
 
+	/** Converts an angle frame attribute to a float frame attribute in degrees */
+	UFUNCTION(BlueprintPure, Category = "AnimDatabase")
+	static UE_API FAnimDatabaseFrameAttribute FrameAttributeAngleToFloatDegrees(const FAnimDatabaseFrameAttribute& AngleFrameAttribute);
+
+	/** Converts an angle frame attribute to a float frame attribute in radians */
+	UFUNCTION(BlueprintPure, Category = "AnimDatabase")
+	static UE_API FAnimDatabaseFrameAttribute FrameAttributeAngleToFloatRadians(const FAnimDatabaseFrameAttribute& AngleFrameAttribute);
+
 	/** Converts a float frame attribute to a bool frame attribute */
 	UFUNCTION(BlueprintPure, Category = "AnimDatabase")
 	static UE_API FAnimDatabaseFrameAttribute FrameAttributeFloatToBool(const FAnimDatabaseFrameAttribute& FloatFrameAttribute);
 
+	/** Converts a float frame attribute to an angle frame attribute in degrees */
+	UFUNCTION(BlueprintPure, Category = "AnimDatabase")
+	static UE_API FAnimDatabaseFrameAttribute FrameAttributeFloatToAngleDegrees(const FAnimDatabaseFrameAttribute& FloatFrameAttribute);
+
+	/** Converts a float frame attribute to an angle frame attribute in radians */
+	UFUNCTION(BlueprintPure, Category = "AnimDatabase")
+	static UE_API FAnimDatabaseFrameAttribute FrameAttributeFloatToAngleRadians(const FAnimDatabaseFrameAttribute& FloatFrameAttribute);
+
 public:
 
 	/** Applies a Gaussian smoothing filter to the values in each range of a frame attribute. Assumes type is compatible */
@@ -1468,10 +1591,26 @@ public:
 	UFUNCTION(BlueprintPure, Category = "AnimDatabase")
 	static UE_API void FrameAttributeFloatAverage(float& OutAverageFloat, const FAnimDatabaseFrameAttribute& FloatFrameAttribute);
 
+	/** Computes the average of an angle frame attribute in degrees */
+	UFUNCTION(BlueprintPure, Category = "AnimDatabase")
+	static UE_API void FrameAttributeAngleAverageDegrees(float& OutAverageAngle, const FAnimDatabaseFrameAttribute& AngleFrameAttribute);
+
+	/** Computes the average of an angle frame attribute in radians */
+	UFUNCTION(BlueprintPure, Category = "AnimDatabase")
+	static UE_API void FrameAttributeAngleAverageRadians(float& OutAverageAngle, const FAnimDatabaseFrameAttribute& AngleFrameAttribute);
+
 	/** Computes the min and max of a float frame attribute */
 	UFUNCTION(BlueprintPure, Category = "AnimDatabase")
 	static UE_API void FrameAttributeFloatMinAndMax(float& OutMin, float& OutMax, const FAnimDatabaseFrameAttribute& FloatFrameAttribute);
 
+	/** Computes the min and max of an angle frame attribute in degrees */
+	UFUNCTION(BlueprintPure, Category = "AnimDatabase")
+	static UE_API void FrameAttributeAngleMinAndMaxDegrees(float& OutMin, float& OutMax, const FAnimDatabaseFrameAttribute& AngleFrameAttribute);
+
+	/** Computes the min and max of an angle frame attribute in radians */
+	UFUNCTION(BlueprintPure, Category = "AnimDatabase")
+	static UE_API void FrameAttributeAngleMinAndMaxRadians(float& OutMin, float& OutMax, const FAnimDatabaseFrameAttribute& AngleFrameAttribute);
+
 	/** Computes the average of a location frame attribute */
 	UFUNCTION(BlueprintPure, Category = "AnimDatabase")
 	static UE_API void FrameAttributeLocationAverage(FVector& OutAverageLocation, const FAnimDatabaseFrameAttribute& LocationFrameAttribute);
@@ -1501,6 +1640,16 @@ public:
 	UFUNCTION(BlueprintPure, Category = "AnimDatabase")
 	static UE_API void FrameAttributeRotationAngleTraveled(float& OutMean, float& OutStd, float& OutMin, float& OutMax, const FAnimDatabaseFrameAttribute& RotationFrameAttribute);
 
+public:
+
+	/** Converts a FrameAttribute object to a string. */
+	UFUNCTION(BlueprintPure, Category = "AnimDatabase", meta = (BlueprintAutocast, CompactNodeTitle = "->", AutoCreateRefTerm = "FrameAttribute"))
+	static UE_API FString FrameAttributeToString(const FAnimDatabaseFrameAttribute& FrameAttribute);
+
+	/** Converts a FrameAttribute object to a string object in full. */
+	UFUNCTION(BlueprintPure, Category = "AnimDatabase")
+	static UE_API FString FrameAttributeToStringFormat(const UAnimDatabase* Database, const FAnimDatabaseFrameAttribute& FrameAttribute, const int32 Cutoff = 11);
+
 
 public:
 
diff --git a/Engine/Plugins/Experimental/Animation/AnimDatabase/Source/AnimDatabase/Public/AnimDatabaseFrameRanges.h b/Engine/Plugins/Experimental/Animation/AnimDatabase/Source/AnimDatabase/Public/AnimDatabaseFrameRanges.h
--- a/Engine/Plugins/Experimental/Animation/AnimDatabase/Source/AnimDatabase/Public/AnimDatabaseFrameRanges.h
+++ b/Engine/Plugins/Experimental/Animation/AnimDatabase/Source/AnimDatabase/Public/AnimDatabaseFrameRanges.h
@@ -193,6 +193,14 @@ public:
 
 	/** Computes the difference of an array of Frames objects using an array view */
 	static UE_API FAnimDatabaseFrames FramesDifferenceFromArrayView(const TArrayView<const FAnimDatabaseFrames> Frames);
+
+	/** Converts a Frames object to a string. */
+	UFUNCTION(BlueprintPure, Category = "AnimDatabase", meta = (BlueprintAutocast, CompactNodeTitle = "->", AutoCreateRefTerm = "Frames"))
+	static UE_API FString FramesToString(const FAnimDatabaseFrames& Frames);
+
+	/** Converts a Frames object to a string object in full. */
+	UFUNCTION(BlueprintPure, Category = "AnimDatabase")
+	static UE_API FString FramesToStringFormat(const UAnimDatabase* Database, const FAnimDatabaseFrames& Frames, const int32 Cutoff = 11);
 };
 
 /** A simple class that has a single implementable function which constructs a FAnimDatabaseFrames object from a UAnimDatabase. */
@@ -606,6 +614,16 @@ public:
 	UFUNCTION(BlueprintPure, Category = "AnimDatabase")
 	static UE_API FAnimDatabaseFrameRanges FrameRangesRangeAtIndex(const FAnimDatabaseFrameRanges& FrameRanges, const int32 RangeIdx);
 
+public:
+
+	/** Converts a FrameRanges object to a string. */
+	UFUNCTION(BlueprintPure, Category = "AnimDatabase", meta = (BlueprintAutocast, CompactNodeTitle = "->", AutoCreateRefTerm = "FrameRanges"))
+	static UE_API FString FrameRangesToString(const FAnimDatabaseFrameRanges& FrameRanges);
+
+	/** Converts a FrameRanges object to a string object in full. */
+	UFUNCTION(BlueprintPure, Category = "AnimDatabase")
+	static UE_API FString FrameRangesToStringFormat(const UAnimDatabase* Database, const FAnimDatabaseFrameRanges& FrameRanges, const int32 Cutoff = 11);
+
 public:
 
 	/** Finds all of the Anim Notify Classes in the given frame ranges */
diff --git a/Engine/Plugins/Experimental/Animation/AnimDatabase/Source/AnimDatabase/Public/AnimDatabaseMath.h b/Engine/Plugins/Experimental/Animation/AnimDatabase/Source/AnimDatabase/Public/AnimDatabaseMath.h
--- a/Engine/Plugins/Experimental/Animation/AnimDatabase/Source/AnimDatabase/Public/AnimDatabaseMath.h
+++ b/Engine/Plugins/Experimental/Animation/AnimDatabase/Source/AnimDatabase/Public/AnimDatabaseMath.h
@@ -187,6 +187,7 @@ namespace UE::AnimDatabase::Math
 	// Hermite Interpolation
 
 	UE_API void HermiteValueInterpolate(float& OutValue, float& OutValueVelocity, const float P0, const float P1, const float V0, const float V1, const float X, const float FrameTime);
+	UE_API void HermiteAngleInterpolate(float& OutAngle, float& OutAngularVelocity, const float A0, const float A1, const float V0, const float V1, const float X, const float FrameTime);
 	UE_API void HermiteLocationInterpolate(FVector& OutPosition, FVector3f& OutLinearVelocity, const FVector P0, const FVector P1, const FVector3f V0, const FVector3f V1, const float X, const float FrameTime);
 	UE_API void HermiteLocationInterpolate(FVector3f& OutPosition, FVector3f& OutLinearVelocity, const FVector3f P0, const FVector3f P1, const FVector3f V0, const FVector3f V1, const float X, const float FrameTime);
 	UE_API void HermiteRotationInterpolate(FQuat4f& OutRotation, FVector3f& OutAngularVelocity, const FQuat4f R0, const FQuat4f R1, const FVector3f V0, const FVector3f V1, const float X, const float FrameTime);
@@ -202,6 +203,7 @@ namespace UE::AnimDatabase::Math
 		const float FrameTime);
 
 	UE_API void HermiteValueInterpolate(float& OutValue, const float P0, const float P1, const float V0, const float V1, const float X, const float FrameTime);
+	UE_API void HermiteAngleInterpolate(float& OutAngle , const float A0, const float A1, const float V0, const float V1, const float X, const float FrameTime);
 	UE_API void HermiteLocationInterpolate(FVector& OutPosition, const FVector P0, const FVector P1, const FVector3f V0, const FVector3f V1, const float X, const float FrameTime);
 	UE_API void HermiteLocationInterpolate(FVector3f& OutPosition, const FVector3f P0, const FVector3f P1, const FVector3f V0, const FVector3f V1, const float X, const float FrameTime);
 	UE_API void HermiteRotationInterpolate(FQuat4f& OutRotation, const FQuat4f R0, const FQuat4f R1, const FVector3f V0, const FVector3f V1, const float X, const float FrameTime);
@@ -218,6 +220,7 @@ namespace UE::AnimDatabase::Math
 	// Linear Interpolation
 
 	UE_API void ValueInterpolateLinear(float& OutValue, float& OutValueVelocity, const float P0, const float P1, const float Alpha, const float FrameTime);
+	UE_API void AngleInterpolateLinear(float& OutAngle, float& OutAngularVelocity, const float A0, const float A1, const float Alpha, const float FrameTime);
 	UE_API void LocationInterpolateLinear(FVector& OutLocation, FVector3f& OutLinearVelocity, const FVector P0, const FVector P1, const float Alpha, const float FrameTime);
 	UE_API void LocationInterpolateLinear(FVector3f& OutLocation, FVector3f& OutLinearVelocity, const FVector3f P0, const FVector3f P1, const float Alpha, const float FrameTime);
 	UE_API void RotationInterpolateLinear(FQuat4f& OutRotation, FVector3f& OutAngularVelocity, const FQuat4f R0, const FQuat4f R1, const float Alpha, const float FrameTime);
@@ -231,6 +234,7 @@ namespace UE::AnimDatabase::Math
 		const float FrameTime);
 
 	UE_API void ValueInterpolateLinear(float& OutValue, const float P0, const float P1, const float Alpha, const float FrameTime);
+	UE_API void AngleInterpolateLinear(float& OutAngle, const float A0, const float A1, const float Alpha, const float FrameTime);
 	UE_API void LocationInterpolateLinear(FVector& OutLocation, const FVector P0, const FVector P1, const float Alpha, const float FrameTime);
 	UE_API void LocationInterpolateLinear(FVector3f& OutLocation, const FVector3f P0, const FVector3f P1, const float Alpha, const float FrameTime);
 	UE_API void RotationInterpolateLinear(FQuat4f& OutRotation, const FQuat4f R0, const FQuat4f R1, const float Alpha, const float FrameTime);
@@ -247,6 +251,7 @@ namespace UE::AnimDatabase::Math
 	// Cubic Interpolation
 
 	UE_API void ValueInterpolateCubic(float& OutValue, float& OutValueVelocity, const float P0, const float P1, const float P2, const float P3, const float Alpha, const float FrameTime);
+	UE_API void AngleInterpolateCubic(float& OutAngle, float& OutAngularVelocity, const float A0, const float A1, const float A2, const float A3, const float Alpha, const float FrameTime);
 	UE_API void LocationInterpolateCubic(FVector& OutLocation, FVector3f& OutLinearVelocity, const FVector P0, const FVector P1, const FVector P2, const FVector P3, const float Alpha, const float FrameTime);
 	UE_API void LocationInterpolateCubic(FVector3f& OutLocation, FVector3f& OutLinearVelocity, const FVector3f P0, const FVector3f P1, const FVector3f P2, const FVector3f P3, const float Alpha, const float FrameTime);
 	UE_API void RotationInterpolateCubic(FQuat4f& OutRotation, FVector3f& OutAngularVelocity, const FQuat4f R0, const FQuat4f R1, const FQuat4f R2, const FQuat4f R3, const float Alpha, const float FrameTime);
@@ -262,6 +267,7 @@ namespace UE::AnimDatabase::Math
 		const float FrameTime);
 
 	UE_API void ValueInterpolateCubic(float& OutValue, const float P0, const float P1, const float P2, const float P3, const float Alpha, const float FrameTime);
+	UE_API void AngleInterpolateCubic(float& OutAngle, const float A0, const float A1, const float A2, const float A3, const float Alpha, const float FrameTime);
 	UE_API void LocationInterpolateCubic(FVector& OutLocation, const FVector P0, const FVector P1, const FVector P2, const FVector P3, const float Alpha, const float FrameTime);
 	UE_API void LocationInterpolateCubic(FVector3f& OutLocation, const FVector3f P0, const FVector3f P1, const FVector3f P2, const FVector3f P3, const float Alpha, const float FrameTime);
 	UE_API void RotationInterpolateCubic(FQuat4f& OutRotation, const FQuat4f R0, const FQuat4f R1, const FQuat4f R2, const FQuat4f R3, const float Alpha, const float FrameTime);
@@ -280,6 +286,7 @@ namespace UE::AnimDatabase::Math
 	// Cubic Interpolation Start
 
 	UE_API void ValueInterpolateCubicStart(float& OutValue, float& OutValueVelocity, const float P1, const float P2, const float P3, const float Alpha, const float FrameTime);
+	UE_API void AngleInterpolateCubicStart(float& OutAngle, float& OutAngularVelocity, const float A1, const float A2, const float A3, const float Alpha, const float FrameTime);
 	UE_API void LocationInterpolateCubicStart(FVector& OutLocation, FVector3f& OutLinearVelocity, const FVector P1, const FVector P2, const FVector P3, const float Alpha, const float FrameTime);
 	UE_API void LocationInterpolateCubicStart(FVector3f& OutLocation, FVector3f& OutLinearVelocity, const FVector3f P1, const FVector3f P2, const FVector3f P3, const float Alpha, const float FrameTime);
 	UE_API void RotationInterpolateCubicStart(FQuat4f& OutRotation, FVector3f& OutAngularVelocity, const FQuat4f R1, const FQuat4f R2, const FQuat4f R3, const float Alpha, const float FrameTime);
@@ -294,6 +301,7 @@ namespace UE::AnimDatabase::Math
 		const float FrameTime);
 
 	UE_API void ValueInterpolateCubicStart(float& OutValue, const float P1, const float P2, const float P3, const float Alpha, const float FrameTime);
+	UE_API void AngleInterpolateCubicStart(float& OutAngle, const float A1, const float A2, const float A3, const float Alpha, const float FrameTime);
 	UE_API void LocationInterpolateCubicStart(FVector& OutLocation, const FVector P1, const FVector P2, const FVector P3, const float Alpha, const float FrameTime);
 	UE_API void LocationInterpolateCubicStart(FVector3f& OutLocation, const FVector3f P1, const FVector3f P2, const FVector3f P3, const float Alpha, const float FrameTime);
 	UE_API void RotationInterpolateCubicStart(FQuat4f& OutRotation, const FQuat4f R1, const FQuat4f R2, const FQuat4f R3, const float Alpha, const float FrameTime);
@@ -311,6 +319,7 @@ namespace UE::AnimDatabase::Math
 	// Cubic Interpolation End
 
 	UE_API void ValueInterpolateCubicEnd(float& OutValue, float& OutValueVelocity, const float P0, const float P1, const float P2, const float Alpha, const float FrameTime);
+	UE_API void AngleInterpolateCubicEnd(float& OutAngle, float& OutAngularVelocity, const float A0, const float A1, const float A2, const float Alpha, const float FrameTime);
 	UE_API void LocationInterpolateCubicEnd(FVector& OutLocation, FVector3f& OutLinearVelocity, const FVector P0, const FVector P1, const FVector P2, const float Alpha, const float FrameTime);
 	UE_API void LocationInterpolateCubicEnd(FVector3f& OutLocation, FVector3f& OutLinearVelocity, const FVector3f P0, const FVector3f P1, const FVector3f P2, const float Alpha, const float FrameTime);
 	UE_API void RotationInterpolateCubicEnd(FQuat4f& OutRotation, FVector3f& OutAngularVelocity, const FQuat4f R0, const FQuat4f R1, const FQuat4f R2, const float Alpha, const float FrameTime);
@@ -325,6 +334,7 @@ namespace UE::AnimDatabase::Math
 		const float FrameTime);
 
 	UE_API void ValueInterpolateCubicEnd(float& OutValue, const float P0, const float P1, const float P2, const float Alpha, const float FrameTime);
+	UE_API void AngleInterpolateCubicEnd(float& OutAngle, const float A0, const float A1, const float A2, const float Alpha, const float FrameTime);
 	UE_API void LocationInterpolateCubicEnd(FVector& OutLocation, const FVector P0, const FVector P1, const FVector P2, const float Alpha, const float FrameTime);
 	UE_API void LocationInterpolateCubicEnd(FVector3f& OutLocation, const FVector3f P0, const FVector3f P1, const FVector3f P2, const float Alpha, const float FrameTime);
 	UE_API void RotationInterpolateCubicEnd(FQuat4f& OutRotation, const FQuat4f R0, const FQuat4f R1, const FQuat4f R2, const float Alpha, const float FrameTime);
@@ -347,6 +357,7 @@ namespace UE::AnimDatabase::Math
 	// Cubic Monotone Interpolation
 
 	UE_API void ValueInterpolateCubicMono(float& OutValue, float& OutValueVelocity, const float P0, const float P1, const float P2, const float P3, const float Alpha, const float FrameTime);
+	UE_API void AngleInterpolateCubicMono(float& OutAngle, float& OutAngularVelocity, const float A0, const float A1, const float A2, const float A3, const float Alpha, const float FrameTime);
 	UE_API void LocationInterpolateCubicMono(FVector& OutLocation, FVector3f& OutLinearVelocity, const FVector P0, const FVector P1, const FVector P2, const FVector P3, const float Alpha, const float FrameTime);
 	UE_API void LocationInterpolateCubicMono(FVector3f& OutLocation, FVector3f& OutLinearVelocity, const FVector3f P0, const FVector3f P1, const FVector3f P2, const FVector3f P3, const float Alpha, const float FrameTime);
 	UE_API void RotationInterpolateCubicMono(FQuat4f& OutRotation, FVector3f& OutAngularVelocity, const FQuat4f R0, const FQuat4f R1, const FQuat4f R2, const FQuat4f R3, const float Alpha, const float FrameTime);
@@ -362,6 +373,7 @@ namespace UE::AnimDatabase::Math
 		const float FrameTime);
 
 	UE_API void ValueInterpolateCubicMono(float& OutValue, const float P0, const float P1, const float P2, const float P3, const float Alpha, const float FrameTime);
+	UE_API void AngleInterpolateCubicMono(float& OutAngle, const float A0, const float A1, const float A2, const float A3, const float Alpha, const float FrameTime);
 	UE_API void LocationInterpolateCubicMono(FVector& OutLocation, const FVector P0, const FVector P1, const FVector P2, const FVector P3, const float Alpha, const float FrameTime);
 	UE_API void LocationInterpolateCubicMono(FVector3f& OutLocation, const FVector3f P0, const FVector3f P1, const FVector3f P2, const FVector3f P3, const float Alpha, const float FrameTime);
 	UE_API void RotationInterpolateCubicMono(FQuat4f& OutRotation, const FQuat4f R0, const FQuat4f R1, const FQuat4f R2, const FQuat4f R3, const float Alpha, const float FrameTime);
@@ -380,6 +392,7 @@ namespace UE::AnimDatabase::Math
 	// Cubic Monotone Interpolation Start
 
 	UE_API void ValueInterpolateCubicMonoStart(float& OutValue, float& OutValueVelocity, const float P1, const float P2, const float P3, const float Alpha, const float FrameTime);
+	UE_API void AngleInterpolateCubicMonoStart(float& OutAngle, float& OutAngularVelocity, const float A1, const float A2, const float A3, const float Alpha, const float FrameTime);
 	UE_API void LocationInterpolateCubicMonoStart(FVector& OutLocation, FVector3f& OutLinearVelocity, const FVector P1, const FVector P2, const FVector P3, const float Alpha, const float FrameTime);
 	UE_API void LocationInterpolateCubicMonoStart(FVector3f& OutLocation, FVector3f& OutLinearVelocity, const FVector3f P1, const FVector3f P2, const FVector3f P3, const float Alpha, const float FrameTime);
 	UE_API void RotationInterpolateCubicMonoStart(FQuat4f& OutRotation, FVector3f& OutAngularVelocity, const FQuat4f R1, const FQuat4f R2, const FQuat4f R3, const float Alpha, const float FrameTime);
@@ -394,6 +407,7 @@ namespace UE::AnimDatabase::Math
 		const float FrameTime);
 
 	UE_API void ValueInterpolateCubicMonoStart(float& OutValue, const float P1, const float P2, const float P3, const float Alpha, const float FrameTime);
+	UE_API void AngleInterpolateCubicMonoStart(float& OutAngle, const float A1, const float A2, const float A3, const float Alpha, const float FrameTime);
 	UE_API void LocationInterpolateCubicMonoStart(FVector& OutLocation, const FVector P1, const FVector P2, const FVector P3, const float Alpha, const float FrameTime);
 	UE_API void LocationInterpolateCubicMonoStart(FVector3f& OutLocation, const FVector3f P1, const FVector3f P2, const FVector3f P3, const float Alpha, const float FrameTime);
 	UE_API void RotationInterpolateCubicMonoStart(FQuat4f& OutRotation, const FQuat4f R1, const FQuat4f R2, const FQuat4f R3, const float Alpha, const float FrameTime);
@@ -411,6 +425,7 @@ namespace UE::AnimDatabase::Math
 	// Cubic Monotone Interpolation End
 
 	UE_API void ValueInterpolateCubicMonoEnd(float& OutValue, float& OutValueVelocity, const float P0, const float P1, const float P2, const float Alpha, const float FrameTime);
+	UE_API void AngleInterpolateCubicMonoEnd(float& OutAngle, float& OutAngularVelocity, const float A0, const float A1, const float A2, const float Alpha, const float FrameTime);
 	UE_API void LocationInterpolateCubicMonoEnd(FVector& OutLocation, FVector3f& OutLinearVelocity, const FVector P0, const FVector P1, const FVector P2, const float Alpha, const float FrameTime);
 	UE_API void LocationInterpolateCubicMonoEnd(FVector3f& OutLocation, FVector3f& OutLinearVelocity, const FVector3f P0, const FVector3f P1, const FVector3f P2, const float Alpha, const float FrameTime);
 	UE_API void RotationInterpolateCubicMonoEnd(FQuat4f& OutRotation, FVector3f& OutAngularVelocity, const FQuat4f R0, const FQuat4f R1, const FQuat4f R2, const float Alpha, const float FrameTime);
@@ -425,6 +440,7 @@ namespace UE::AnimDatabase::Math
 		const float FrameTime);
 
 	UE_API void ValueInterpolateCubicMonoEnd(float& OutValue, const float P0, const float P1, const float P2, const float Alpha, const float FrameTime);
+	UE_API void AngleInterpolateCubicMonoEnd(float& OutAngle, const float A0, const float A1, const float A2, const float Alpha, const float FrameTime);
 	UE_API void LocationInterpolateCubicMonoEnd(FVector& OutLocation, const FVector P0, const FVector P1, const FVector P2, const float Alpha, const float FrameTime);
 	UE_API void LocationInterpolateCubicMonoEnd(FVector3f& OutLocation, const FVector3f P0, const FVector3f P1, const FVector3f P2, const float Alpha, const float FrameTime);
 	UE_API void RotationInterpolateCubicMonoEnd(FQuat4f& OutRotation, const FQuat4f R0, const FQuat4f R1, const FQuat4f R2, const float Alpha, const float FrameTime);
@@ -488,6 +504,11 @@ namespace UE::AnimDatabase::Math
 		const TLearningArrayView<1, const float> InValues,
 		const int32 SampleFrame);
 
+	UE_API void AngleSampleNearest(
+		float& OutAngle,
+		const TLearningArrayView<1, const float> InAngles,
+		const int32 SampleFrame);
+
 	UE_API void LocationSampleNearest(
 		FVector& OutLocation,
 		const TLearningArrayView<1, const FVector> InLocations,
@@ -562,6 +583,14 @@ namespace UE::AnimDatabase::Math
 		const float SampleAlpha,
 		const float FrameDeltaTime);
 
+	UE_API void AngleSampleLinear(
+		float& OutAngle,
+		const TLearningArrayView<1, const float> InAngles,
+		const int32 SampleFrame0,
+		const int32 SampleFrame1,
+		const float SampleAlpha,
+		const float FrameDeltaTime);
+
 	UE_API void LocationSampleLinear(
 		FVector& OutLocation,
 		const TLearningArrayView<1, const FVector> InLocations,
@@ -665,6 +694,16 @@ namespace UE::AnimDatabase::Math
 		const float SampleAlpha,
 		const float FrameDeltaTime);
 
+	UE_API void AngleSampleCubic(
+		float& OutAngle,
+		const TLearningArrayView<1, const float> InAngles,
+		const int32 SampleFrame0,
+		const int32 SampleFrame1,
+		const int32 SampleFrame2,
+		const int32 SampleFrame3,
+		const float SampleAlpha,
+		const float FrameDeltaTime);
+
 	UE_API void LocationSampleCubic(
 		FVector& OutLocation,
 		const TLearningArrayView<1, const FVector> InLocations,
@@ -786,6 +825,16 @@ namespace UE::AnimDatabase::Math
 		const float SampleAlpha,
 		const float FrameDeltaTime);
 
+	UE_API void AngleSampleCubicMono(
+		float& OutAngle,
+		const TLearningArrayView<1, const float> InAngles,
+		const int32 SampleFrame0,
+		const int32 SampleFrame1,
+		const int32 SampleFrame2,
+		const int32 SampleFrame3,
+		const float SampleAlpha,
+		const float FrameDeltaTime);
+
 	UE_API void LocationSampleCubicMono(
 		FVector& OutLocation,
 		const TLearningArrayView<1, const FVector> InLocations,
diff --git a/Engine/Plugins/Experimental/Animation/AnimDatabase/Source/AnimDatabase/Public/AnimDatabasePose.h b/Engine/Plugins/Experimental/Animation/AnimDatabase/Source/AnimDatabase/Public/AnimDatabasePose.h
--- a/Engine/Plugins/Experimental/Animation/AnimDatabase/Source/AnimDatabase/Public/AnimDatabasePose.h
+++ b/Engine/Plugins/Experimental/Animation/AnimDatabase/Source/AnimDatabase/Public/AnimDatabasePose.h
@@ -216,6 +216,7 @@ namespace UE::AnimDatabase
 
 		UE_API bool GetBool(const int32 FrameIdx, const int32 AttributeIdx) const;
 		UE_API float GetFloat(const int32 FrameIdx, const int32 AttributeIdx) const;
+		UE_API float GetAngle(const int32 FrameIdx, const int32 AttributeIdx) const;
 		UE_API FVector3f GetLocation(const int32 FrameIdx, const int32 AttributeIdx) const;
 		UE_API FQuat4f GetRotation(const int32 FrameIdx, const int32 AttributeIdx) const;
 		UE_API FVector3f GetScale(const int32 FrameIdx, const int32 AttributeIdx) const;
@@ -249,6 +250,7 @@ namespace UE::AnimDatabase
 
 		UE_API bool GetBool(const int32 FrameIdx, const int32 AttributeIdx) const;
 		UE_API float GetFloat(const int32 FrameIdx, const int32 AttributeIdx) const;
+		UE_API float GetAngle(const int32 FrameIdx, const int32 AttributeIdx) const;
 		UE_API FVector3f GetLocation(const int32 FrameIdx, const int32 AttributeIdx) const;
 		UE_API FQuat4f GetRotation(const int32 FrameIdx, const int32 AttributeIdx) const;
 		UE_API FVector3f GetScale(const int32 FrameIdx, const int32 AttributeIdx) const;
@@ -263,6 +265,7 @@ namespace UE::AnimDatabase
 
 		UE_API void SetBool(const int32 FrameIdx, const int32 AttributeIdx, const bool bValue) const;
 		UE_API void SetFloat(const int32 FrameIdx, const int32 AttributeIdx, const float Value) const;
+		UE_API void SetAngle(const int32 FrameIdx, const int32 AttributeIdx, const float Value) const;
 		UE_API void SetLocation(const int32 FrameIdx, const int32 AttributeIdx, const FVector3f Location) const;
 		UE_API void SetRotation(const int32 FrameIdx, const int32 AttributeIdx, const FQuat4f Rotation) const;
 		UE_API void SetScale(const int32 FrameIdx, const int32 AttributeIdx, const FVector3f Scale) const;
@@ -302,6 +305,7 @@ namespace UE::AnimDatabase
 
 		UE_API bool GetBool(const int32 FrameIdx, const int32 AttributeIdx) const;
 		UE_API float GetFloat(const int32 FrameIdx, const int32 AttributeIdx) const;
+		UE_API float GetAngle(const int32 FrameIdx, const int32 AttributeIdx) const;
 		UE_API FVector3f GetLocation(const int32 FrameIdx, const int32 AttributeIdx) const;
 		UE_API FQuat4f GetRotation(const int32 FrameIdx, const int32 AttributeIdx) const;
 		UE_API FVector3f GetScale(const int32 FrameIdx, const int32 AttributeIdx) const;
@@ -736,6 +740,14 @@ public:
 	UFUNCTION(BlueprintPure, Category = "AnimDatabase")
 	static UE_API float PoseStateFloatAttribute(const FAnimDatabasePoseState& PoseState, const int32 AttributeIndex);
 
+	/** Get the given attribute as a angle in degrees */
+	UFUNCTION(BlueprintPure, Category = "AnimDatabase")
+	static UE_API float PoseStateAngleAttributeDegrees(const FAnimDatabasePoseState& PoseState, const int32 AttributeIndex);
+
+	/** Get the given attribute as a angle in radians */
+	UFUNCTION(BlueprintPure, Category = "AnimDatabase")
+	static UE_API float PoseStateAngleAttributeRadians(const FAnimDatabasePoseState& PoseState, const int32 AttributeIndex);
+
 	/** Get the given attribute as a location */
 	UFUNCTION(BlueprintPure, Category = "AnimDatabase")
 	static UE_API FVector PoseStateLocationAttribute(const FAnimDatabasePoseState& PoseState, const int32 AttributeIndex);
