AnimDatabase: Fix for attach bone setting function

#Synced-CL 51024568

diff --git a/Engine/Plugins/Experimental/Animation/AnimDatabase/Source/AnimDatabase/Private/AnimDatabase.cpp b/Engine/Plugins/Experimental/Animation/AnimDatabase/Source/AnimDatabase/Private/AnimDatabase.cpp
--- a/Engine/Plugins/Experimental/Animation/AnimDatabase/Source/AnimDatabase/Private/AnimDatabase.cpp
+++ b/Engine/Plugins/Experimental/Animation/AnimDatabase/Source/AnimDatabase/Private/AnimDatabase.cpp
@@ -6051,8 +6051,6 @@ void UAnimDatabaseFunction_UpdateIKBones::Run_Implementation(UAnimDatabase* InDa
 
 void UAnimDatabaseFunction_SetRootFromLookAt::Run_Implementation(UAnimDatabase* InDatabase, const FAnimDatabaseFrameRanges& InFrameRanges, const bool bShouldTransact)
 {
-
-
 	const FAnimDatabaseFrameAttribute SmoothedRootTransform = bUseSavGolFilter ?
 		
 		UAnimDatabaseFrameAttributeLibrary::FrameAttributeSavGolSmoothedRootTransform(
@@ -6251,10 +6249,10 @@ void UAnimDatabaseFunction_SetBoneToRootTransformAtFrames::Run_Implementation(UA
 		InDatabase,
 		BoneName,
 		UAnimDatabaseFrameAttributeLibrary::FrameAttributeMultiplyInverse(
-			UAnimDatabaseFrameAttributeLibrary::MakeBoneGlobalTransformFrameAttribute(InDatabase, InFrameRanges, ParentIndex),
 			UAnimDatabaseFrameAttributeLibrary::MakeTransformFrameAttributeNoScale(
 				UAnimDatabaseFrameAttributeLibrary::MakeLocationFrameAttributeFromConstant(InFrameRanges, AvgRootLocation),
-				UAnimDatabaseFrameAttributeLibrary::MakeRotationFrameAttributeFromConstant(InFrameRanges, AvgRootDirection.ToOrientationRotator()))),
+				UAnimDatabaseFrameAttributeLibrary::MakeRotationFrameAttributeFromConstant(InFrameRanges, AvgRootDirection.ToOrientationRotator())),
+			UAnimDatabaseFrameAttributeLibrary::MakeBoneGlobalTransformFrameAttribute(InDatabase, InFrameRanges, ParentIndex)),
 		bShouldTransact);
 
 }
