AnimGen: Refactoring in preparation for CR

#Synced-CL 51066804

diff --git a/Engine/Plugins/Experimental/Animation/AnimDatabase/Source/AnimDatabase/Private/AnimDatabasePose.cpp b/Engine/Plugins/Experimental/Animation/AnimDatabase/Source/AnimDatabase/Private/AnimDatabasePose.cpp
--- a/Engine/Plugins/Experimental/Animation/AnimDatabase/Source/AnimDatabase/Private/AnimDatabasePose.cpp
+++ b/Engine/Plugins/Experimental/Animation/AnimDatabase/Source/AnimDatabase/Private/AnimDatabasePose.cpp
@@ -602,6 +602,11 @@ namespace UE::AnimDatabase
 		AttributeData[FrameIdx][AttributeOffsets[AttributeIdx] + 0] = Value;
 	}
 
+	void FPoseAttributeDataView::SetAngle(const int32 FrameIdx, const int32 AttributeIdx, const float Value) const
+	{
+		AttributeData[FrameIdx][AttributeOffsets[AttributeIdx] + 0] = Value;
+	}
+
 	void FPoseAttributeDataView::SetLocation(const int32 FrameIdx, const int32 AttributeIdx, const FVector3f Location) const
 	{
 		AttributeData[FrameIdx][AttributeOffsets[AttributeIdx] + 0] = Location.X;
@@ -2495,7 +2500,27 @@ namespace UE::AnimDatabase
 			const TLearningArrayView<1, const int32> BoneParents,
 			const Learning::FIndexSet BoneIndices)
 		{
-			TRACE_CPUPROFILER_EVENT_SCOPE(AnimDatabase::PoseData::ForwardKinematics);
+			ForwardKinematicsPartial(OutPoseGlobalBoneData, InPoseLocalBoneData, InPoseRootData, BoneParents, BoneIndices, Learning::FIndexSet(0, BoneIndices.Num()));
+		}
+
+		void ForwardKinematics(
+			const FPoseGlobalBoneDataView& OutGlobalPoseData,
+			const FPoseLocalBoneDataConstView& InPoseLocalBoneData,
+			const FPoseRootDataConstView& InPoseRootData,
+			const TLearningArrayView<1, const int32> BoneParents)
+		{
+			ForwardKinematics(OutGlobalPoseData, InPoseLocalBoneData, InPoseRootData, BoneParents, Learning::FIndexSet(0, BoneParents.Num()));
+		}
+
+		void ForwardKinematicsPartial(
+			const FPoseGlobalBoneDataView& OutPoseGlobalBoneData,
+			const FPoseLocalBoneDataConstView& InPoseLocalBoneData,
+			const FPoseRootDataConstView& InPoseRootData,
+			const TLearningArrayView<1, const int32> BoneParents,
+			const UE::Learning::FIndexSet BoneIndices,
+			const UE::Learning::FIndexSet BonesToUpdate)
+		{
+			TRACE_CPUPROFILER_EVENT_SCOPE(AnimDatabase::PoseData::ForwardKinematicsPartial);
 
 			const int32 FrameNum = InPoseLocalBoneData.GetFrameNum();
 			const int32 BoneNum = InPoseLocalBoneData.GetBoneNum();
@@ -2511,7 +2536,7 @@ namespace UE::AnimDatabase
 
 			for (int32 FrameIdx = 0; FrameIdx < FrameNum; FrameIdx++)
 			{
-				for (int32 BoneIdx = 0; BoneIdx < BoneNum; BoneIdx++)
+				for (int32 BoneIdx : BonesToUpdate)
 				{
 					const FVector3f CurrBoneLocation = InPoseLocalBoneData.BoneLocations[FrameIdx][BoneIdx];
 					const FQuat4f CurrBoneRotation = InPoseLocalBoneData.BoneRotations[FrameIdx][BoneIdx];
@@ -2531,7 +2556,7 @@ namespace UE::AnimDatabase
 
 					if (RefParentBoneIdx == INDEX_NONE)
 					{
-						ParentBoneLocation =InPoseRootData.RootLocations[FrameIdx];
+						ParentBoneLocation = InPoseRootData.RootLocations[FrameIdx];
 						ParentBoneRotation = InPoseRootData.RootRotations[FrameIdx];
 						ParentBoneScale = InPoseRootData.RootScales[FrameIdx];
 						ParentBoneLinearVelocity = InPoseRootData.RootLinearVelocities[FrameIdx];
@@ -2543,6 +2568,7 @@ namespace UE::AnimDatabase
 						const int32 ParentBoneIdx = RefToBoneIndices[RefParentBoneIdx];
 						check(ParentBoneIdx != INDEX_NONE);
 						check(ParentBoneIdx < BoneIdx);
+						check(BonesToUpdate.Contains(ParentBoneIdx));
 
 						ParentBoneLocation = OutPoseGlobalBoneData.BoneLocations[FrameIdx][ParentBoneIdx];
 						ParentBoneRotation = OutPoseGlobalBoneData.BoneRotations[FrameIdx][ParentBoneIdx];
@@ -2565,13 +2591,1177 @@ namespace UE::AnimDatabase
 			}
 		}
 
-		void ForwardKinematics(
-			const FPoseGlobalBoneDataView& OutGlobalPoseData,
-			const FPoseLocalBoneDataConstView& InPoseLocalBoneData,
-			const FPoseRootDataConstView& InPoseRootData,
-			const TLearningArrayView<1, const int32> BoneParents)
+		void Teleport(
+			const FPoseRootDataView& InOutPoseRootData,
+			const TLearningArrayView<1, const FVector> RootLocations,
+			const TLearningArrayView<1, const FQuat4f> RootRotations)
 		{
-			ForwardKinematics(OutGlobalPoseData, InPoseLocalBoneData, InPoseRootData, BoneParents, Learning::FIndexSet(0, BoneParents.Num()));
+			check(InOutPoseRootData.GetFrameNum() == RootLocations.Num());
+			check(InOutPoseRootData.GetFrameNum() == RootRotations.Num());
+
+			const int32 FrameNum = InOutPoseRootData.GetFrameNum();
+
+			for (int32 FrameIdx = 0; FrameIdx < FrameNum; FrameIdx++)
+			{
+				const FVector RootLocationDifference = RootLocations[FrameIdx] - InOutPoseRootData.RootLocations[FrameIdx];
+				const FQuat4f RootRotationDifference = RootRotations[FrameIdx] * InOutPoseRootData.RootRotations[FrameIdx].Inverse();
+
+				InOutPoseRootData.RootLocations[FrameIdx] = RootLocationDifference + InOutPoseRootData.RootLocations[FrameIdx];
+				InOutPoseRootData.RootRotations[FrameIdx] = RootRotationDifference * InOutPoseRootData.RootRotations[FrameIdx];
+				InOutPoseRootData.RootLinearVelocities[FrameIdx] = RootRotationDifference.RotateVector(InOutPoseRootData.RootLinearVelocities[FrameIdx]);
+				InOutPoseRootData.RootAngularVelocities[FrameIdx] = RootRotationDifference.RotateVector(InOutPoseRootData.RootAngularVelocities[FrameIdx]);
+			}
+		}
+
+		void BlendRoot(
+			const FPoseRootDataView& OutPoseData,
+			const FPoseRootDataConstView& InPoseA,
+			const FPoseRootDataConstView& InPoseB,
+			const float Alpha)
+		{
+			TRACE_CPUPROFILER_EVENT_SCOPE(AnimDatabase::PoseData::Blend);
+
+			check(OutPoseData.GetFrameNum() == InPoseA.GetFrameNum());
+			check(OutPoseData.GetFrameNum() == InPoseB.GetFrameNum());
+
+			const int32 FrameNum = OutPoseData.GetFrameNum();
+
+			for (int32 FrameIdx = 0; FrameIdx < FrameNum; FrameIdx++)
+			{
+				OutPoseData.RootLocations[FrameIdx] = FMath::Lerp(InPoseA.RootLocations[FrameIdx], InPoseB.RootLocations[FrameIdx], Alpha);
+				OutPoseData.RootRotations[FrameIdx] = FQuat4f::Slerp(InPoseA.RootRotations[FrameIdx], InPoseB.RootRotations[FrameIdx], Alpha);
+				OutPoseData.RootScales[FrameIdx] = Math::VectorEerp(InPoseA.RootScales[FrameIdx], InPoseB.RootScales[FrameIdx], Alpha);
+			}
+		}
+
+		void IntegrateRoot(const FPoseRootDataView& InOutRootData, const float DeltaTime)
+		{
+			const int32 FrameNum = InOutRootData.GetFrameNum();
+
+			for (int32 FrameIdx = 1; FrameIdx < FrameNum; FrameIdx++)
+			{
+				const FVector3f RootLocalLinearVelocity = InOutRootData.RootRotations[FrameIdx].UnrotateVector(InOutRootData.RootLinearVelocities[FrameIdx]);
+				const FVector3f RootLocalAngularVelocity = InOutRootData.RootRotations[FrameIdx].UnrotateVector(InOutRootData.RootAngularVelocities[FrameIdx]);
+
+				InOutRootData.RootLocations[FrameIdx] = InOutRootData.RootLocations[FrameIdx - 1] + DeltaTime * (FVector)InOutRootData.RootLinearVelocities[FrameIdx - 1];
+				InOutRootData.RootRotations[FrameIdx] = FQuat4f::MakeFromRotationVector(DeltaTime * InOutRootData.RootAngularVelocities[FrameIdx - 1]) * InOutRootData.RootRotations[FrameIdx - 1];
+				InOutRootData.RootLinearVelocities[FrameIdx] = InOutRootData.RootRotations[FrameIdx].RotateVector(RootLocalLinearVelocity);
+				InOutRootData.RootAngularVelocities[FrameIdx] = InOutRootData.RootRotations[FrameIdx].RotateVector(RootLocalAngularVelocity);
+			}
+		}
+
+		void BlendVia(
+			const FPoseLocalBoneDataView& OutPoseData,
+			const TLearningArrayView<2, FQuat4f> InOutViaRotations,
+			const FPoseLocalBoneDataConstView& InPoseA,
+			const FPoseLocalBoneDataConstView& InPoseB,
+			const float Alpha)
+		{
+			TRACE_CPUPROFILER_EVENT_SCOPE(AnimDatabase::PoseData::BlendVia);
+
+			check(OutPoseData.GetFrameNum() == InPoseA.GetFrameNum());
+			check(OutPoseData.GetFrameNum() == InPoseB.GetFrameNum());
+			check(OutPoseData.GetBoneNum() == InPoseA.GetBoneNum());
+			check(OutPoseData.GetBoneNum() == InPoseB.GetBoneNum());
+			check(OutPoseData.GetFrameNum() == InOutViaRotations.Num<0>());
+			check(OutPoseData.GetBoneNum() == InOutViaRotations.Num<1>());
+
+			const int32 FrameNum = OutPoseData.GetFrameNum();
+			const int32 BoneNum = OutPoseData.GetBoneNum();
+
+			for (int32 FrameIdx = 0; FrameIdx < FrameNum; FrameIdx++)
+			{
+				for (int32 BoneIdx = 0; BoneIdx < BoneNum; BoneIdx++)
+				{
+					const FQuat4f RotationDiff = (InPoseB.BoneRotations[FrameIdx][BoneIdx] * InPoseA.BoneRotations[FrameIdx][BoneIdx].Inverse()).GetShortestArcWith(InOutViaRotations[FrameIdx][BoneIdx]);
+
+					InOutViaRotations[FrameIdx][BoneIdx] = RotationDiff;
+
+					OutPoseData.BoneLocations[FrameIdx][BoneIdx] = FMath::Lerp(InPoseA.BoneLocations[FrameIdx][BoneIdx], InPoseB.BoneLocations[FrameIdx][BoneIdx], Alpha);
+					OutPoseData.BoneRotations[FrameIdx][BoneIdx] = FQuat4f::MakeFromRotationVector(RotationDiff.ToRotationVector() * Alpha) * InPoseA.BoneRotations[FrameIdx][BoneIdx];
+					OutPoseData.BoneScales[FrameIdx][BoneIdx] = Math::VectorEerp(InPoseA.BoneScales[FrameIdx][BoneIdx], InPoseB.BoneScales[FrameIdx][BoneIdx], Alpha);
+				}
+			}
+		}
+
+		void FitExtrapolationHalfLives(
+			const TLearningArrayView<2, FVector3f> OutBoneLocationHalfLives,
+			const TLearningArrayView<2, FVector3f> OutBoneRotationHalfLives,
+			const TLearningArrayView<2, FVector3f> OutBoneScaleHalfLives,
+			const FPoseLocalBoneDataConstView& InPoseSource,
+			const FPoseLocalBoneDataConstView& InPoseTarget,
+			const float ExtrapolationHalfLife,
+			const float ExtrapolationHalfLifeMin,
+			const float ExtrapolationHalfLifeMax,
+			const float MaxLinearVelocity,
+			const float MaxAngularVelocity,
+			const float MaxScalarVelocity)
+		{
+			TRACE_CPUPROFILER_EVENT_SCOPE(AnimDatabase::PoseData::FitExtrapolationHalfLives);
+
+			check(InPoseSource.GetFrameNum() == InPoseTarget.GetFrameNum());
+			check(InPoseSource.GetBoneNum() == InPoseTarget.GetBoneNum());
+			check(InPoseSource.GetFrameNum() == OutBoneLocationHalfLives.Num<0>());
+			check(InPoseSource.GetFrameNum() == OutBoneRotationHalfLives.Num<0>());
+			check(InPoseSource.GetFrameNum() == OutBoneScaleHalfLives.Num<0>());
+			check(InPoseSource.GetBoneNum() == OutBoneLocationHalfLives.Num<1>());
+			check(InPoseSource.GetBoneNum() == OutBoneRotationHalfLives.Num<1>());
+			check(InPoseSource.GetBoneNum() == OutBoneScaleHalfLives.Num<1>());
+
+			const int32 FrameNum = InPoseSource.GetFrameNum();
+			const int32 BoneNum = InPoseSource.GetBoneNum();
+
+			for (int32 FrameIdx = 0; FrameIdx < FrameNum; FrameIdx++)
+			{
+				for (int32 BoneIdx = 0; BoneIdx < BoneNum; BoneIdx++)
+				{
+					const FVector3f ClampedBoneLinearVelocity = InPoseSource.BoneLinearVelocities[FrameIdx][BoneIdx].GetClampedToMaxSize(MaxLinearVelocity);
+					const FVector3f ClampedBoneAngularVelocity = InPoseSource.BoneAngularVelocities[FrameIdx][BoneIdx].GetClampedToMaxSize(MaxAngularVelocity);
+					const FVector3f ClampedBoneScalarVelocity = InPoseSource.BoneScalarVelocities[FrameIdx][BoneIdx].GetClampedToMaxSize(MaxScalarVelocity);
+
+					const FVector3f BoneLocationDifference = InPoseTarget.BoneLocations[FrameIdx][BoneIdx] - InPoseSource.BoneLocations[FrameIdx][BoneIdx];
+					const FVector3f BoneRotationDifference = (InPoseTarget.BoneRotations[FrameIdx][BoneIdx] * InPoseSource.BoneRotations[FrameIdx][BoneIdx].Inverse()).GetShortestArcWith(FQuat4f::Identity).ToRotationVector();
+					const FVector3f BoneScaleDifference = Math::VectorDivMax(InPoseTarget.BoneScales[FrameIdx][BoneIdx], InPoseSource.BoneScales[FrameIdx][BoneIdx]);
+
+					OutBoneLocationHalfLives[FrameIdx][BoneIdx] = Math::ComputeDecayHalfLifeFromDiffAndVelocity(
+						BoneLocationDifference,
+						ClampedBoneLinearVelocity,
+						ExtrapolationHalfLife,
+						ExtrapolationHalfLifeMin,
+						ExtrapolationHalfLifeMax);
+
+					OutBoneRotationHalfLives[FrameIdx][BoneIdx] = Math::ComputeDecayHalfLifeFromDiffAndVelocity(
+						BoneRotationDifference,
+						ClampedBoneAngularVelocity,
+						ExtrapolationHalfLife,
+						ExtrapolationHalfLifeMin,
+						ExtrapolationHalfLifeMax);
+
+					OutBoneScaleHalfLives[FrameIdx][BoneIdx] = Math::ComputeDecayHalfLifeFromDiffAndVelocity(
+						BoneScaleDifference,
+						ClampedBoneScalarVelocity,
+						ExtrapolationHalfLife,
+						ExtrapolationHalfLifeMin,
+						ExtrapolationHalfLifeMax);
+				}
+			}
+		}
+
+		void ExtrapolateRoot(
+			const FPoseRootDataView& OutPoseData,
+			const FPoseRootDataConstView& InPoseData,
+			const float Time,
+			const float RootHalfLife)
+		{
+			TRACE_CPUPROFILER_EVENT_SCOPE(AnimDatabase::PoseData::ExtrapolateRoot);
+
+			check(OutPoseData.GetFrameNum() == InPoseData.GetFrameNum());
+
+			const int32 FrameNum = OutPoseData.GetFrameNum();
+
+			const FVector3f RootHalfLifeVector = FVector3f(RootHalfLife, RootHalfLife, RootHalfLife);
+
+			for (int32 FrameIdx = 0; FrameIdx < FrameNum; FrameIdx++)
+			{
+				Math::ExtrapolateTranslation(
+					OutPoseData.RootLocations[FrameIdx],
+					OutPoseData.RootLinearVelocities[FrameIdx],
+					InPoseData.RootLocations[FrameIdx],
+					InPoseData.RootLinearVelocities[FrameIdx],
+					Time,
+					RootHalfLifeVector);
+
+				Math::ExtrapolateRotation(
+					OutPoseData.RootRotations[FrameIdx],
+					OutPoseData.RootAngularVelocities[FrameIdx],
+					InPoseData.RootRotations[FrameIdx],
+					InPoseData.RootAngularVelocities[FrameIdx],
+					Time,
+					RootHalfLifeVector);
+
+				Math::ExtrapolateScale(
+					OutPoseData.RootScales[FrameIdx],
+					OutPoseData.RootScalarVelocities[FrameIdx],
+					InPoseData.RootScales[FrameIdx],
+					InPoseData.RootScalarVelocities[FrameIdx],
+					Time,
+					RootHalfLifeVector);
+			}
+		}
+
+		void Extrapolate(
+			const FPoseLocalBoneDataView& OutPoseData,
+			const FPoseLocalBoneDataConstView& InPoseData,
+			const float Time,
+			const TLearningArrayView<2, const FVector3f> BoneLocationHalfLives,
+			const TLearningArrayView<2, const FVector3f> BoneRotationHalfLives,
+			const TLearningArrayView<2, const FVector3f> BoneScaleHalfLives)
+		{
+			TRACE_CPUPROFILER_EVENT_SCOPE(AnimDatabase::PoseData::Extrapolate);
+
+			check(OutPoseData.GetFrameNum() == InPoseData.GetFrameNum());
+			check(OutPoseData.GetBoneNum() == InPoseData.GetBoneNum());
+			check(OutPoseData.GetFrameNum() == BoneLocationHalfLives.Num<0>());
+			check(OutPoseData.GetFrameNum() == BoneRotationHalfLives.Num<0>());
+			check(OutPoseData.GetFrameNum() == BoneScaleHalfLives.Num<0>());
+			check(OutPoseData.GetBoneNum() == BoneLocationHalfLives.Num<1>());
+			check(OutPoseData.GetBoneNum() == BoneRotationHalfLives.Num<1>());
+			check(OutPoseData.GetBoneNum() == BoneScaleHalfLives.Num<1>());
+
+			const int32 FrameNum = OutPoseData.GetFrameNum();
+			const int32 BoneNum = OutPoseData.GetBoneNum();
+
+			for (int32 FrameIdx = 0; FrameIdx < FrameNum; FrameIdx++)
+			{
+				for (int32 BoneIdx = 0; BoneIdx < BoneNum; BoneIdx++)
+				{
+					Math::ExtrapolateTranslation(
+						OutPoseData.BoneLocations[FrameIdx][BoneIdx],
+						OutPoseData.BoneLinearVelocities[FrameIdx][BoneIdx],
+						InPoseData.BoneLocations[FrameIdx][BoneIdx],
+						InPoseData.BoneLinearVelocities[FrameIdx][BoneIdx],
+						Time,
+						BoneLocationHalfLives[FrameIdx][BoneIdx]);
+
+					Math::ExtrapolateRotation(
+						OutPoseData.BoneRotations[FrameIdx][BoneIdx],
+						OutPoseData.BoneAngularVelocities[FrameIdx][BoneIdx],
+						InPoseData.BoneRotations[FrameIdx][BoneIdx],
+						InPoseData.BoneAngularVelocities[FrameIdx][BoneIdx],
+						Time,
+						BoneRotationHalfLives[FrameIdx][BoneIdx]);
+
+					Math::ExtrapolateScale(
+						OutPoseData.BoneScales[FrameIdx][BoneIdx],
+						OutPoseData.BoneScalarVelocities[FrameIdx][BoneIdx],
+						InPoseData.BoneScales[FrameIdx][BoneIdx],
+						InPoseData.BoneScalarVelocities[FrameIdx][BoneIdx],
+						Time,
+						BoneScaleHalfLives[FrameIdx][BoneIdx]);
+				}
+			}
+		}
+
+		void FitInertializationOffsetsRoot(
+			const FPoseRootDataView& OutOffsetPose,
+			const FPoseRootDataConstView& InPoseSource,
+			const FPoseRootDataConstView& InPoseTarget)
+		{
+			TRACE_CPUPROFILER_EVENT_SCOPE(AnimDatabase::PoseData::FitInertializationOffsetsRoot);
+
+			check(OutOffsetPose.GetFrameNum() == InPoseSource.GetFrameNum());
+			check(OutOffsetPose.GetFrameNum() == InPoseTarget.GetFrameNum());
+
+			const int32 FrameNum = OutOffsetPose.GetFrameNum();
+
+			for (int32 FrameIdx = 0; FrameIdx < FrameNum; FrameIdx++)
+			{
+				const FVector RootLocationDifference = InPoseSource.RootLocations[FrameIdx] - InPoseTarget.RootLocations[FrameIdx];
+				const FQuat4f RootRotationDifference = (InPoseSource.RootRotations[FrameIdx] * InPoseTarget.RootRotations[FrameIdx].Inverse()).GetShortestArcWith(FQuat4f::Identity);
+				const FVector3f RootScaleDifference = Math::VectorDivMax(InPoseSource.RootScales[FrameIdx], InPoseTarget.RootScales[FrameIdx]);
+
+				OutOffsetPose.RootLocations[FrameIdx] = RootLocationDifference;
+				OutOffsetPose.RootRotations[FrameIdx] = RootRotationDifference;
+				OutOffsetPose.RootScales[FrameIdx] = RootScaleDifference;
+				OutOffsetPose.RootLinearVelocities[FrameIdx] = InPoseSource.RootLinearVelocities[FrameIdx] - InPoseTarget.RootLinearVelocities[FrameIdx];
+				OutOffsetPose.RootAngularVelocities[FrameIdx] = InPoseSource.RootAngularVelocities[FrameIdx] - InPoseTarget.RootAngularVelocities[FrameIdx];
+				OutOffsetPose.RootScalarVelocities[FrameIdx] = InPoseSource.RootScalarVelocities[FrameIdx] - InPoseTarget.RootScalarVelocities[FrameIdx];
+			}
+		}
+
+		void ApplyInertializationOffsetsRoot(
+			const FPoseRootDataView& OutFinalPose,
+			const FPoseRootDataConstView& InOffsetPose,
+			const FPoseRootDataConstView& InPoseTarget,
+			const float BlendTime,
+			const float BlendDuration)
+		{
+			TRACE_CPUPROFILER_EVENT_SCOPE(AnimDatabase::PoseData::ApplyInertializationOffsetsRoot);
+
+			check(OutFinalPose.GetFrameNum() == InOffsetPose.GetFrameNum());
+			check(OutFinalPose.GetFrameNum() == InPoseTarget.GetFrameNum());
+
+			const int32 FrameNum = OutFinalPose.GetFrameNum();
+
+			// Cubic Weights
+			const float RT = FMath::Clamp(BlendTime / (BlendDuration + UE_SMALL_NUMBER), 0.0f, 1.0f);
+			const float RW0 = 2.0f * RT * RT * RT - 3.0f * RT * RT + 1.0f;
+			const float RW1 = (RT * RT * RT - 2.0f * RT * RT + RT) * BlendDuration;
+			const float RW2 = (6.0f * RT * RT - 6.0f * RT) / (BlendDuration + UE_SMALL_NUMBER);
+			const float RW3 = 3.0f * RT * RT - 4.0f * RT + 1.0f;
+
+			for (int32 FrameIdx = 0; FrameIdx < FrameNum; FrameIdx++)
+			{
+				const FVector RootLocationOffset = RW0 * InOffsetPose.RootLocations[FrameIdx] + RW1 * (FVector)InOffsetPose.RootLinearVelocities[FrameIdx];
+				const FQuat4f RootRotationOffset = FQuat4f::MakeFromRotationVector(
+					RW0 * InOffsetPose.RootRotations[FrameIdx].ToRotationVector() + RW1 * InOffsetPose.RootAngularVelocities[FrameIdx]);
+				const FVector3f RootScaleOffset =Math::VectorExp(RW0 * Math::VectorLogSafe(InOffsetPose.RootScales[FrameIdx]) + RW1 * InOffsetPose.RootScalarVelocities[FrameIdx]);
+				const FVector3f RootLinearVelocityOffset = RW2 * (FVector3f)InOffsetPose.RootLocations[FrameIdx] + RW3 * InOffsetPose.RootLinearVelocities[FrameIdx];
+				const FVector3f RootAngularVelocityOffset = RW2 * InOffsetPose.RootRotations[FrameIdx].ToRotationVector() +
+					RW3 * InOffsetPose.RootAngularVelocities[FrameIdx];
+				const FVector3f RootScalarVelocityOffset = RW2 * Math::VectorLogSafe(InOffsetPose.RootScales[FrameIdx]) + RW3 * InOffsetPose.RootScalarVelocities[FrameIdx];
+
+				OutFinalPose.RootLocations[FrameIdx] = RootLocationOffset + InPoseTarget.RootLocations[FrameIdx];
+				OutFinalPose.RootRotations[FrameIdx] = RootRotationOffset * InPoseTarget.RootRotations[FrameIdx];
+				OutFinalPose.RootScales[FrameIdx] = RootScaleOffset * InPoseTarget.RootScales[FrameIdx];
+				OutFinalPose.RootLinearVelocities[FrameIdx] = RootLinearVelocityOffset + InPoseTarget.RootLinearVelocities[FrameIdx];
+				OutFinalPose.RootAngularVelocities[FrameIdx] = RootAngularVelocityOffset + RootRotationOffset.RotateVector(InPoseTarget.RootAngularVelocities[FrameIdx]);
+				OutFinalPose.RootScalarVelocities[FrameIdx] = RootScalarVelocityOffset + InPoseTarget.RootScalarVelocities[FrameIdx];
+			}
+		}
+
+		void FitInertializationOffsets(
+			const FPoseLocalBoneDataView& OutOffsetPose,
+			const FPoseLocalBoneDataConstView& InPoseSource,
+			const FPoseLocalBoneDataConstView& InPoseTarget)
+		{
+			TRACE_CPUPROFILER_EVENT_SCOPE(AnimDatabase::PoseData::FitInertializationOffsets);
+
+			check(OutOffsetPose.GetFrameNum() == InPoseSource.GetFrameNum());
+			check(OutOffsetPose.GetBoneNum() == InPoseSource.GetBoneNum());
+			check(OutOffsetPose.GetFrameNum() == InPoseTarget.GetFrameNum());
+			check(OutOffsetPose.GetBoneNum() == InPoseTarget.GetBoneNum());
+
+			const int32 FrameNum = OutOffsetPose.GetFrameNum();
+			const int32 BoneNum = OutOffsetPose.GetBoneNum();
+
+			for (int32 FrameIdx = 0; FrameIdx < FrameNum; FrameIdx++)
+			{
+				for (int32 BoneIdx = 0; BoneIdx < BoneNum; BoneIdx++)
+				{
+					const FVector3f BoneLocationDifference = InPoseSource.BoneLocations[FrameIdx][BoneIdx] - InPoseTarget.BoneLocations[FrameIdx][BoneIdx];
+					const FQuat4f BoneRotationDifference = (InPoseSource.BoneRotations[FrameIdx][BoneIdx] * InPoseTarget.BoneRotations[FrameIdx][BoneIdx].Inverse()).GetShortestArcWith(FQuat4f::Identity);
+					const FVector3f BoneScaleDifference = Math::VectorDivMax(InPoseSource.BoneScales[FrameIdx][BoneIdx], InPoseTarget.BoneScales[FrameIdx][BoneIdx]);
+
+					OutOffsetPose.BoneLocations[FrameIdx][BoneIdx] = BoneLocationDifference;
+					OutOffsetPose.BoneRotations[FrameIdx][BoneIdx] = BoneRotationDifference;
+					OutOffsetPose.BoneScales[FrameIdx][BoneIdx] = BoneScaleDifference;
+					OutOffsetPose.BoneLinearVelocities[FrameIdx][BoneIdx] = InPoseSource.BoneLinearVelocities[FrameIdx][BoneIdx] - InPoseTarget.BoneLinearVelocities[FrameIdx][BoneIdx];
+					OutOffsetPose.BoneAngularVelocities[FrameIdx][BoneIdx] = InPoseSource.BoneAngularVelocities[FrameIdx][BoneIdx] - InPoseTarget.BoneAngularVelocities[FrameIdx][BoneIdx];
+					OutOffsetPose.BoneScalarVelocities[FrameIdx][BoneIdx] = InPoseSource.BoneScalarVelocities[FrameIdx][BoneIdx] - InPoseTarget.BoneScalarVelocities[FrameIdx][BoneIdx];
+				}
+			}
+		}
+
+		void ApplyInertializationOffsets(
+			const FPoseLocalBoneDataView& OutFinalPose,
+			const FPoseLocalBoneDataConstView& InOffsetPose,
+			const FPoseLocalBoneDataConstView& InPoseTarget,
+			const float BlendTime,
+			const float BlendDuration)
+		{
+			TRACE_CPUPROFILER_EVENT_SCOPE(AnimDatabase::PoseData::ApplyInertializationOffsets);
+
+			check(OutFinalPose.GetFrameNum() == InOffsetPose.GetFrameNum());
+			check(OutFinalPose.GetBoneNum() == InOffsetPose.GetBoneNum());
+			check(OutFinalPose.GetFrameNum() == InPoseTarget.GetFrameNum());
+			check(OutFinalPose.GetBoneNum() == InPoseTarget.GetBoneNum());
+
+			const int32 FrameNum = OutFinalPose.GetFrameNum();
+			const int32 BoneNum = OutFinalPose.GetBoneNum();
+
+			// Cubic Weights
+			const float T = FMath::Clamp(BlendTime / (BlendDuration + UE_SMALL_NUMBER), 0.0f, 1.0f);
+			const float W0 = 2.0f * T * T * T - 3.0f * T * T + 1.0f;
+			const float W1 = (T * T * T - 2.0f * T * T + T) * BlendDuration;
+			const float W2 = (6.0f * T * T - 6.0f * T) / (BlendDuration + UE_SMALL_NUMBER);
+			const float W3 = 3.0f * T * T - 4.0f * T + 1.0f;
+
+			for (int32 FrameIdx = 0; FrameIdx < FrameNum; FrameIdx++)
+			{
+				for (int32 BoneIdx = 0; BoneIdx < BoneNum; BoneIdx++)
+				{
+					const FVector3f BoneLocationOffset = W0 * InOffsetPose.BoneLocations[FrameIdx][BoneIdx] + W1 * InOffsetPose.BoneLinearVelocities[FrameIdx][BoneIdx];
+					const FQuat4f BoneRotationOffset = FQuat4f::MakeFromRotationVector(
+						W0 * InOffsetPose.BoneRotations[FrameIdx][BoneIdx].ToRotationVector() + W1 * InOffsetPose.BoneAngularVelocities[FrameIdx][BoneIdx]);
+					const FVector3f BoneScaleOffset = Math::VectorExp(W0 * Math::VectorLogSafe(InOffsetPose.BoneScales[FrameIdx][BoneIdx]) + W1 * InOffsetPose.BoneScalarVelocities[FrameIdx][BoneIdx]);
+					const FVector3f BoneLinearVelocityOffset = W2 * InOffsetPose.BoneLocations[FrameIdx][BoneIdx] + W3 * InOffsetPose.BoneLinearVelocities[FrameIdx][BoneIdx];
+					const FVector3f BoneAngularVelocityOffset = W2 * InOffsetPose.BoneRotations[FrameIdx][BoneIdx].ToRotationVector() +
+						W3 * InOffsetPose.BoneAngularVelocities[FrameIdx][BoneIdx];
+					const FVector3f BoneScalarVelocityOffset = W2 * Math::VectorLogSafe(InOffsetPose.BoneScales[FrameIdx][BoneIdx]) + W3 * InOffsetPose.BoneScalarVelocities[FrameIdx][BoneIdx];
+
+					OutFinalPose.BoneLocations[FrameIdx][BoneIdx] = BoneLocationOffset + InPoseTarget.BoneLocations[FrameIdx][BoneIdx];
+					OutFinalPose.BoneRotations[FrameIdx][BoneIdx] = BoneRotationOffset * InPoseTarget.BoneRotations[FrameIdx][BoneIdx];
+					OutFinalPose.BoneScales[FrameIdx][BoneIdx] = BoneScaleOffset * InPoseTarget.BoneScales[FrameIdx][BoneIdx];
+					OutFinalPose.BoneLinearVelocities[FrameIdx][BoneIdx] = BoneLinearVelocityOffset + InPoseTarget.BoneLinearVelocities[FrameIdx][BoneIdx];
+					OutFinalPose.BoneAngularVelocities[FrameIdx][BoneIdx] = BoneAngularVelocityOffset + BoneRotationOffset.RotateVector(InPoseTarget.BoneAngularVelocities[FrameIdx][BoneIdx]);
+					OutFinalPose.BoneScalarVelocities[FrameIdx][BoneIdx] = BoneScalarVelocityOffset + InPoseTarget.BoneScalarVelocities[FrameIdx][BoneIdx];
+				}
+			}
+		}
+
+		void FitInertializationOffsetsAttributes(
+			const FPoseAttributeDataView& OutOffsetPose,
+			const FPoseAttributeDataConstView& InPoseSource,
+			const FPoseAttributeDataConstView& InPoseTarget)
+		{
+			TRACE_CPUPROFILER_EVENT_SCOPE(AnimDatabase::PoseData::FitInertializationOffsetsAttributes);
+
+			check(OutOffsetPose.GetFrameNum() == InPoseSource.GetFrameNum());
+			check(UE::Learning::Array::Equal(OutOffsetPose.GetAttributeTypes(), InPoseSource.GetAttributeTypes()));
+			check(OutOffsetPose.GetFrameNum() == InPoseTarget.GetFrameNum());
+			check(UE::Learning::Array::Equal(OutOffsetPose.GetAttributeTypes(), InPoseTarget.GetAttributeTypes()));
+
+			const int32 FrameNum = OutOffsetPose.GetFrameNum();
+			const int32 AttributeNum = OutOffsetPose.GetAttributeNum();
+
+			for (int32 FrameIdx = 0; FrameIdx < FrameNum; FrameIdx++)
+			{
+				for (int32 AttributeIdx = 0; AttributeIdx < AttributeNum; AttributeIdx++)
+				{
+					const EAnimDatabaseAttributeType AttributeType = InPoseSource.GetAttributeType(AttributeIdx);
+					const bool bSourceActive = InPoseSource.GetAttributeActive(FrameIdx, AttributeIdx);
+					const bool bTargetActive = InPoseTarget.GetAttributeActive(FrameIdx, AttributeIdx);
+
+					OutOffsetPose.SetAttributeActive(FrameIdx, AttributeIdx, bSourceActive && bTargetActive);
+
+					if (OutOffsetPose.GetAttributeActive(FrameIdx, AttributeIdx))
+					{
+						switch (AttributeType)
+						{
+						case EAnimDatabaseAttributeType::Null: break;
+						case EAnimDatabaseAttributeType::Bool: OutOffsetPose.SetBool(FrameIdx, AttributeIdx, false); break;
+						case EAnimDatabaseAttributeType::Float: OutOffsetPose.SetFloat(FrameIdx, AttributeIdx, InPoseSource.GetFloat(FrameIdx, AttributeIdx) - InPoseTarget.GetFloat(FrameIdx, AttributeIdx)); break;
+						case EAnimDatabaseAttributeType::Angle: OutOffsetPose.SetAngle(FrameIdx, AttributeIdx, FMath::FindDeltaAngleRadians(InPoseTarget.GetAngle(FrameIdx, AttributeIdx), InPoseSource.GetAngle(FrameIdx, AttributeIdx))); break;
+						case EAnimDatabaseAttributeType::Location: OutOffsetPose.SetLocation(FrameIdx, AttributeIdx, InPoseSource.GetLocation(FrameIdx, AttributeIdx) - InPoseTarget.GetLocation(FrameIdx, AttributeIdx)); break;
+						case EAnimDatabaseAttributeType::Rotation: OutOffsetPose.SetRotation(FrameIdx, AttributeIdx, (InPoseSource.GetRotation(FrameIdx, AttributeIdx) * InPoseTarget.GetRotation(FrameIdx, AttributeIdx).Inverse()).GetShortestArcWith(FQuat4f::Identity)); break;
+						case EAnimDatabaseAttributeType::Scale: OutOffsetPose.SetScale(FrameIdx, AttributeIdx, Math::VectorDivMax(InPoseSource.GetScale(FrameIdx, AttributeIdx), InPoseTarget.GetScale(FrameIdx, AttributeIdx))); break;
+						case EAnimDatabaseAttributeType::LinearVelocity: OutOffsetPose.SetLinearVelocity(FrameIdx, AttributeIdx, InPoseSource.GetLinearVelocity(FrameIdx, AttributeIdx) - InPoseTarget.GetLinearVelocity(FrameIdx, AttributeIdx)); break;
+						case EAnimDatabaseAttributeType::AngularVelocity: OutOffsetPose.SetAngularVelocity(FrameIdx, AttributeIdx, InPoseSource.GetAngularVelocity(FrameIdx, AttributeIdx) - InPoseTarget.GetAngularVelocity(FrameIdx, AttributeIdx)); break;
+						case EAnimDatabaseAttributeType::ScalarVelocity: OutOffsetPose.SetScalarVelocity(FrameIdx, AttributeIdx, InPoseSource.GetScalarVelocity(FrameIdx, AttributeIdx) - InPoseTarget.GetScalarVelocity(FrameIdx, AttributeIdx)); break;
+						case EAnimDatabaseAttributeType::Direction: OutOffsetPose.SetDirection(FrameIdx, AttributeIdx, InPoseSource.GetDirection(FrameIdx, AttributeIdx) - InPoseTarget.GetDirection(FrameIdx, AttributeIdx)); break;
+						case EAnimDatabaseAttributeType::Transform:
+						{
+							const FTransform3f TargetTransform = InPoseTarget.GetTransform(FrameIdx, AttributeIdx);
+							const FTransform3f SourceTransform = InPoseSource.GetTransform(FrameIdx, AttributeIdx);
+							FTransform3f OutTransform = FTransform3f::Identity;
+							OutTransform.SetLocation(SourceTransform.GetLocation() - TargetTransform.GetLocation());
+							OutTransform.SetRotation((SourceTransform.GetRotation() * TargetTransform.GetRotation().Inverse()).GetShortestArcWith(FQuat4f::Identity));
+							OutTransform.SetScale3D(Math::VectorDivMax(SourceTransform.GetScale3D(), TargetTransform.GetScale3D()));
+							OutOffsetPose.SetTransform(FrameIdx, AttributeIdx, OutTransform);
+							break;
+						}
+
+						case EAnimDatabaseAttributeType::Event:
+						{
+							bool bTargetEventKnown = false, bSourceEventKnown = false;
+							float TargetTimeUntilEvent = UE_MAX_FLT, SourceTimeUntilEvent = UE_MAX_FLT;
+							InPoseTarget.GetEvent(bTargetEventKnown, TargetTimeUntilEvent, FrameIdx, AttributeIdx);
+							InPoseSource.GetEvent(bSourceEventKnown, SourceTimeUntilEvent, FrameIdx, AttributeIdx);
+
+							if (bTargetEventKnown && bSourceEventKnown)
+							{
+								OutOffsetPose.SetEvent(FrameIdx, AttributeIdx, true, SourceTimeUntilEvent - TargetTimeUntilEvent);
+							}
+							else
+							{
+								OutOffsetPose.SetEvent(FrameIdx, AttributeIdx, false, UE_MAX_FLT);
+							}
+						}
+						}
+					}
+					else
+					{
+						UE::Learning::Array::Zero(
+							OutOffsetPose.AttributeData[FrameIdx].Slice(
+								OutOffsetPose.GetAttributeOffset(AttributeIdx),
+								OutOffsetPose.GetAttributeSize(AttributeIdx)));
+					}
+				}
+			}
+		}
+
+		void ApplyInertializationOffsetsAttributes(
+			const FPoseAttributeDataView& OutFinalPose,
+			const FPoseAttributeDataConstView& InOffsetPose,
+			const FPoseAttributeDataConstView& InPoseTarget,
+			const float BlendTime,
+			const float BlendDuration)
+		{
+			TRACE_CPUPROFILER_EVENT_SCOPE(AnimDatabase::PoseData::ApplyInertializationOffsetsAttributes);
+
+			check(OutFinalPose.GetFrameNum() == InOffsetPose.GetFrameNum());
+			check(Learning::Array::Equal(OutFinalPose.GetAttributeTypes(), InOffsetPose.GetAttributeTypes()));
+			check(OutFinalPose.GetFrameNum() == InPoseTarget.GetFrameNum());
+			check(Learning::Array::Equal(OutFinalPose.GetAttributeTypes(), InPoseTarget.GetAttributeTypes()));
+
+			const int32 FrameNum = OutFinalPose.GetFrameNum();
+			const int32 AttributeNum = OutFinalPose.GetAttributeNum();
+
+			const float T = FMath::Clamp(BlendTime / (BlendDuration + UE_SMALL_NUMBER), 0.0f, 1.0f);
+			const float Alpha = 1.0f - T * T * (3.0f - 2.0f * T);
+
+			for (int32 FrameIdx = 0; FrameIdx < FrameNum; FrameIdx++)
+			{
+				for (int32 AttributeIdx = 0; AttributeIdx < AttributeNum; AttributeIdx++)
+				{
+					const EAnimDatabaseAttributeType AttributeType = InPoseTarget.GetAttributeType(AttributeIdx);
+					const bool bTargetActive = InPoseTarget.GetAttributeActive(FrameIdx, AttributeIdx);
+					const bool bOffsetActive = InOffsetPose.GetAttributeActive(FrameIdx, AttributeIdx);
+
+					OutFinalPose.SetAttributeActive(FrameIdx, AttributeIdx, bTargetActive);
+
+					if (bTargetActive && bOffsetActive)
+					{
+						switch (AttributeType)
+						{
+						case EAnimDatabaseAttributeType::Null: break;
+						case EAnimDatabaseAttributeType::Bool: OutFinalPose.SetBool(FrameIdx, AttributeIdx, InPoseTarget.GetBool(FrameIdx, AttributeIdx)); break;
+						case EAnimDatabaseAttributeType::Float: OutFinalPose.SetFloat(FrameIdx, AttributeIdx, Alpha * InOffsetPose.GetFloat(FrameIdx, AttributeIdx) + InPoseTarget.GetFloat(FrameIdx, AttributeIdx)); break;
+						case EAnimDatabaseAttributeType::Angle: OutFinalPose.SetAngle(FrameIdx, AttributeIdx, Alpha * InOffsetPose.GetAngle(FrameIdx, AttributeIdx) + InPoseTarget.GetAngle(FrameIdx, AttributeIdx)); break;
+						case EAnimDatabaseAttributeType::Location: OutFinalPose.SetLocation(FrameIdx, AttributeIdx, Alpha * InOffsetPose.GetLocation(FrameIdx, AttributeIdx) + InPoseTarget.GetLocation(FrameIdx, AttributeIdx)); break;
+						case EAnimDatabaseAttributeType::Rotation: OutFinalPose.SetRotation(FrameIdx, AttributeIdx, FQuat4f::MakeFromRotationVector(Alpha * InOffsetPose.GetRotation(FrameIdx, AttributeIdx).ToRotationVector()) * InPoseTarget.GetRotation(FrameIdx, AttributeIdx)); break;
+						case EAnimDatabaseAttributeType::Scale: OutFinalPose.SetScale(FrameIdx, AttributeIdx, Math::VectorExp(Alpha * Math::VectorLogSafe(InOffsetPose.GetScale(FrameIdx, AttributeIdx))) * InPoseTarget.GetScale(FrameIdx, AttributeIdx)); break;
+						case EAnimDatabaseAttributeType::LinearVelocity: OutFinalPose.SetLinearVelocity(FrameIdx, AttributeIdx, Alpha * InOffsetPose.GetLinearVelocity(FrameIdx, AttributeIdx) + InPoseTarget.GetLinearVelocity(FrameIdx, AttributeIdx)); break;
+						case EAnimDatabaseAttributeType::AngularVelocity: OutFinalPose.SetAngularVelocity(FrameIdx, AttributeIdx, Alpha * InOffsetPose.GetAngularVelocity(FrameIdx, AttributeIdx) + InPoseTarget.GetAngularVelocity(FrameIdx, AttributeIdx)); break;
+						case EAnimDatabaseAttributeType::ScalarVelocity: OutFinalPose.SetScalarVelocity(FrameIdx, AttributeIdx, Alpha * InOffsetPose.GetScalarVelocity(FrameIdx, AttributeIdx) + InPoseTarget.GetScalarVelocity(FrameIdx, AttributeIdx)); break;
+						case EAnimDatabaseAttributeType::Direction: OutFinalPose.SetDirection(FrameIdx, AttributeIdx, (Alpha * InOffsetPose.GetDirection(FrameIdx, AttributeIdx) + InPoseTarget.GetDirection(FrameIdx, AttributeIdx)).GetSafeNormal(UE_SMALL_NUMBER, FVector3f::ForwardVector)); break;
+						case EAnimDatabaseAttributeType::Transform:
+						{
+							const FTransform3f TargetTransform = InPoseTarget.GetTransform(FrameIdx, AttributeIdx);
+							const FTransform3f OffsetTransform = InOffsetPose.GetTransform(FrameIdx, AttributeIdx);
+							FTransform3f OutTransform = FTransform3f::Identity;
+							OutTransform.SetLocation(Alpha * OffsetTransform.GetLocation() + TargetTransform.GetLocation());
+							OutTransform.SetRotation(FQuat4f::MakeFromRotationVector(Alpha * OffsetTransform.GetRotation().ToRotationVector()) * TargetTransform.GetRotation());
+							OutTransform.SetScale3D(Math::VectorExp(Alpha * Math::VectorLogSafe(OffsetTransform.GetScale3D())) * TargetTransform.GetScale3D());
+							OutFinalPose.SetTransform(FrameIdx, AttributeIdx, OutTransform);
+							break;
+						}
+
+						case EAnimDatabaseAttributeType::Event:
+						{
+							bool bTargetEventKnown = false, bOffsetEventKnown = false;
+							float TargetTimeUntilEvent = UE_MAX_FLT, OffsetTimeUntilEvent = UE_MAX_FLT;
+							InPoseTarget.GetEvent(bTargetEventKnown, TargetTimeUntilEvent, FrameIdx, AttributeIdx);
+							InOffsetPose.GetEvent(bOffsetEventKnown, OffsetTimeUntilEvent, FrameIdx, AttributeIdx);
+
+							if (bTargetEventKnown)
+							{
+								if (bOffsetEventKnown)
+								{
+									OutFinalPose.SetEvent(FrameIdx, AttributeIdx, true, Alpha * OffsetTimeUntilEvent + TargetTimeUntilEvent);
+								}
+								else
+								{
+									OutFinalPose.SetEvent(FrameIdx, AttributeIdx, true, TargetTimeUntilEvent);
+								}
+							}
+							else
+							{
+								OutFinalPose.SetEvent(FrameIdx, AttributeIdx, false, UE_MAX_FLT);
+							}
+							break;
+						}
+						default:
+							checkNoEntry();
+						}
+					}
+					else
+					{
+						UE::Learning::Array::Copy(
+							OutFinalPose.AttributeData[FrameIdx].Slice(
+								OutFinalPose.GetAttributeOffset(AttributeIdx),
+								OutFinalPose.GetAttributeSize(AttributeIdx)),
+							InPoseTarget.AttributeData[FrameIdx].Slice(
+								InPoseTarget.GetAttributeOffset(AttributeIdx),
+								InPoseTarget.GetAttributeSize(AttributeIdx)));
+					}
+				}
+			}
+		}
+
+		void UpdateInplaceUsingVelocityIntegrationMix(
+			const FPoseRootDataView& InOutPose,
+			const FPoseRootDataConstView& InTargetPose,
+			const float DeltaTime,
+			const float VelocityAlpha)
+		{
+			TRACE_CPUPROFILER_EVENT_SCOPE(AnimDatabase::PoseData::UpdateInplaceUsingVelocityIntegrationMix);
+
+			check(InOutPose.GetFrameNum() == InTargetPose.GetFrameNum());
+			check(InOutPose.GetFrameNum() == InTargetPose.GetFrameNum());
+
+			const int32 FrameNum = InOutPose.GetFrameNum();
+
+			for (int32 FrameIdx = 0; FrameIdx < FrameNum; FrameIdx++)
+			{
+				const FVector RootLocationIntegrated = InOutPose.RootLocations[FrameIdx] + DeltaTime * (FVector)InTargetPose.RootLinearVelocities[FrameIdx];
+				const FQuat4f RootRotationIntegrated = FQuat4f::MakeFromRotationVector(DeltaTime * InTargetPose.RootAngularVelocities[FrameIdx]) * InOutPose.RootRotations[FrameIdx];
+				const FVector3f RootScaleIntegrated = Math::VectorExpSafe(DeltaTime * InTargetPose.RootScalarVelocities[FrameIdx]) * InOutPose.RootScales[FrameIdx];
+
+				InOutPose.RootLocations[FrameIdx] = FMath::Lerp(InTargetPose.RootLocations[FrameIdx], RootLocationIntegrated, VelocityAlpha);
+				InOutPose.RootRotations[FrameIdx] = FQuat4f::Slerp(InTargetPose.RootRotations[FrameIdx], RootRotationIntegrated, VelocityAlpha);
+				InOutPose.RootScales[FrameIdx] = Math::VectorEerp(InTargetPose.RootScales[FrameIdx], RootScaleIntegrated, VelocityAlpha);
+				InOutPose.RootLinearVelocities[FrameIdx] = InTargetPose.RootLinearVelocities[FrameIdx];
+				InOutPose.RootAngularVelocities[FrameIdx] = InTargetPose.RootAngularVelocities[FrameIdx];
+				InOutPose.RootScalarVelocities[FrameIdx] = InTargetPose.RootScalarVelocities[FrameIdx];
+			}
+		}
+
+		void UpdateInplaceUsingVelocityIntegrationMix(
+			const FPoseLocalBoneDataView& InOutPose,
+			const FPoseLocalBoneDataConstView& InTargetPose,
+			const float DeltaTime,
+			const float VelocityAlpha)
+		{
+			TRACE_CPUPROFILER_EVENT_SCOPE(AnimDatabase::PoseData::UpdateInplaceUsingVelocityIntegrationMix);
+
+			check(InOutPose.GetFrameNum() == InTargetPose.GetFrameNum());
+			check(InOutPose.GetBoneNum() == InTargetPose.GetBoneNum());
+			check(InOutPose.GetFrameNum() == InTargetPose.GetFrameNum());
+			check(InOutPose.GetBoneNum() == InTargetPose.GetBoneNum());
+
+			const int32 FrameNum = InOutPose.GetFrameNum();
+			const int32 BoneNum = InOutPose.GetBoneNum();
+
+			for (int32 FrameIdx = 0; FrameIdx < FrameNum; FrameIdx++)
+			{
+				for (int32 BoneIdx = 0; BoneIdx < BoneNum; BoneIdx++)
+				{
+					const FVector3f BoneLocationIntegrated = DeltaTime * InTargetPose.BoneLinearVelocities[FrameIdx][BoneIdx] + InOutPose.BoneLocations[FrameIdx][BoneIdx];
+					const FQuat4f BoneRotationIntegrated = FQuat4f::MakeFromRotationVector(DeltaTime * InTargetPose.BoneAngularVelocities[FrameIdx][BoneIdx]) * InOutPose.BoneRotations[FrameIdx][BoneIdx];
+					const FVector3f BoneScaleIntegrated = Math::VectorExpSafe(DeltaTime * InTargetPose.BoneScalarVelocities[FrameIdx][BoneIdx]) * InOutPose.BoneScales[FrameIdx][BoneIdx];
+
+					InOutPose.BoneLocations[FrameIdx][BoneIdx] = FMath::Lerp(InTargetPose.BoneLocations[FrameIdx][BoneIdx], BoneLocationIntegrated, VelocityAlpha);
+					InOutPose.BoneRotations[FrameIdx][BoneIdx] = FQuat4f::Slerp(InTargetPose.BoneRotations[FrameIdx][BoneIdx], BoneRotationIntegrated, VelocityAlpha);
+					InOutPose.BoneScales[FrameIdx][BoneIdx] = Math::VectorEerp(InTargetPose.BoneScales[FrameIdx][BoneIdx], BoneScaleIntegrated, VelocityAlpha);
+					InOutPose.BoneLinearVelocities[FrameIdx][BoneIdx] = InTargetPose.BoneLinearVelocities[FrameIdx][BoneIdx];
+					InOutPose.BoneAngularVelocities[FrameIdx][BoneIdx] = InTargetPose.BoneAngularVelocities[FrameIdx][BoneIdx];
+					InOutPose.BoneScalarVelocities[FrameIdx][BoneIdx] = InTargetPose.BoneScalarVelocities[FrameIdx][BoneIdx];
+				}
+			}
+		}
+
+		void DampInplaceUsingVelocityIntegrationMix(
+			const FPoseLocalBoneDataView& InOutPose,
+			const FPoseLocalBoneDataConstView& InTargetPose,
+			const float DeltaTime,
+			const float PoseSmoothingTime)
+		{
+			TRACE_CPUPROFILER_EVENT_SCOPE(AnimDatabase::PoseData::DampInplaceUsingVelocityIntegrationMix);
+
+			check(InOutPose.GetFrameNum() == InTargetPose.GetFrameNum());
+			check(InOutPose.GetBoneNum() == InTargetPose.GetBoneNum());
+			check(InOutPose.GetFrameNum() == InTargetPose.GetFrameNum());
+			check(InOutPose.GetBoneNum() == InTargetPose.GetBoneNum());
+
+			const int32 FrameNum = InOutPose.GetFrameNum();
+			const int32 BoneNum = InOutPose.GetBoneNum();
+
+			for (int32 FrameIdx = 0; FrameIdx < FrameNum; FrameIdx++)
+			{
+				for (int32 BoneIdx = 0; BoneIdx < BoneNum; BoneIdx++)
+				{
+					// Doing the integration with a 50/50 velocity mix between start and end pose is more "implicit"
+					const FVector3f LinearVelocity = FMath::Lerp(InTargetPose.BoneLinearVelocities[FrameIdx][BoneIdx], InOutPose.BoneLinearVelocities[FrameIdx][BoneIdx], 0.5f);
+					const FVector3f AngularVelocity = FMath::Lerp(InTargetPose.BoneAngularVelocities[FrameIdx][BoneIdx], InOutPose.BoneAngularVelocities[FrameIdx][BoneIdx], 0.5f);
+					const FVector3f ScalarVelocity = FMath::Lerp(InTargetPose.BoneScalarVelocities[FrameIdx][BoneIdx], InOutPose.BoneScalarVelocities[FrameIdx][BoneIdx], 0.5f);
+
+					const FVector3f BoneLocationIntegrated = DeltaTime * LinearVelocity + InOutPose.BoneLocations[FrameIdx][BoneIdx];
+					const FQuat4f BoneRotationIntegrated = FQuat4f::MakeFromRotationVector(DeltaTime * AngularVelocity) * InOutPose.BoneRotations[FrameIdx][BoneIdx];
+					const FVector3f BoneScaleIntegrated = Math::VectorExpSafe(DeltaTime * ScalarVelocity) * InOutPose.BoneScales[FrameIdx][BoneIdx];
+
+					const float Factor = 1.0f - FMath::InvExpApprox(DeltaTime / FMath::Max(PoseSmoothingTime, UE_SMALL_NUMBER));
+
+					InOutPose.BoneLocations[FrameIdx][BoneIdx] = FMath::Lerp(BoneLocationIntegrated, InTargetPose.BoneLocations[FrameIdx][BoneIdx], Factor);
+					InOutPose.BoneRotations[FrameIdx][BoneIdx] = FQuat4f::Slerp(BoneRotationIntegrated, InTargetPose.BoneRotations[FrameIdx][BoneIdx], Factor);
+					InOutPose.BoneScales[FrameIdx][BoneIdx] = Math::VectorEerp(BoneScaleIntegrated, InTargetPose.BoneScales[FrameIdx][BoneIdx], Factor);
+					InOutPose.BoneLinearVelocities[FrameIdx][BoneIdx] = InTargetPose.BoneLinearVelocities[FrameIdx][BoneIdx];
+					InOutPose.BoneAngularVelocities[FrameIdx][BoneIdx] = InTargetPose.BoneAngularVelocities[FrameIdx][BoneIdx];
+					InOutPose.BoneScalarVelocities[FrameIdx][BoneIdx] = InTargetPose.BoneScalarVelocities[FrameIdx][BoneIdx];
+				}
+			}
+		}
+
+		void DampInplaceUsingAdaptiveVelocityIntegrationMix(
+			const FPoseLocalBoneDataView& InOutPose,
+			const FPoseLocalBoneDataConstView& InTargetPose,
+			const float DeltaTime,
+			const float MinPoseSmoothingTime,
+			const float MaxPoseSmoothingTime,
+			const float LinearVelocitySmoothingThreshold,
+			const float AngularVelocitySmoothingThreshold,
+			const float ScalarVelocitySmoothingThreshold)
+		{
+			TRACE_CPUPROFILER_EVENT_SCOPE(AnimDatabase::PoseData::DampInplaceUsingAdaptiveVelocityIntegrationMix);
+
+			const float ClampedMinPoseSmoothingTime = FMath::Max(MinPoseSmoothingTime, UE_SMALL_NUMBER);
+			const float ClampedMaxPoseSmoothingTime = FMath::Max(MaxPoseSmoothingTime, UE_SMALL_NUMBER);
+
+			check(InOutPose.GetFrameNum() == InTargetPose.GetFrameNum());
+			check(InOutPose.GetBoneNum() == InTargetPose.GetBoneNum());
+			check(InOutPose.GetFrameNum() == InTargetPose.GetFrameNum());
+			check(InOutPose.GetBoneNum() == InTargetPose.GetBoneNum());
+
+			const int32 FrameNum = InOutPose.GetFrameNum();
+			const int32 BoneNum = InOutPose.GetBoneNum();
+
+			for (int32 FrameIdx = 0; FrameIdx < FrameNum; FrameIdx++)
+			{
+				for (int32 BoneIdx = 0; BoneIdx < BoneNum; BoneIdx++)
+				{
+					// Doing the integration with a 50/50 velocity mix between start and end pose is more "implicit"
+					const FVector3f LinearVelocity = FMath::Lerp(InTargetPose.BoneLinearVelocities[FrameIdx][BoneIdx], InOutPose.BoneLinearVelocities[FrameIdx][BoneIdx], 0.5f);
+					const FVector3f AngularVelocity = FMath::Lerp(InTargetPose.BoneAngularVelocities[FrameIdx][BoneIdx], InOutPose.BoneAngularVelocities[FrameIdx][BoneIdx], 0.5f);
+					const FVector3f ScalarVelocity = FMath::Lerp(InTargetPose.BoneScalarVelocities[FrameIdx][BoneIdx], InOutPose.BoneScalarVelocities[FrameIdx][BoneIdx], 0.5f);
+
+					const FVector3f BoneLocationIntegrated = DeltaTime * LinearVelocity + InOutPose.BoneLocations[FrameIdx][BoneIdx];
+					const FQuat4f BoneRotationIntegrated = FQuat4f::MakeFromRotationVector(DeltaTime * AngularVelocity) * InOutPose.BoneRotations[FrameIdx][BoneIdx];
+					const FVector3f BoneScaleIntegrated = Math::VectorExpSafe(DeltaTime * ScalarVelocity) * InOutPose.BoneScales[FrameIdx][BoneIdx];
+
+					const float LinearVelocitySmoothingTime = (ClampedMaxPoseSmoothingTime * LinearVelocitySmoothingThreshold) / FMath::Max(LinearVelocity.Length(), UE_SMALL_NUMBER);
+					const float AngularVelocitySmoothingTime = (ClampedMaxPoseSmoothingTime * AngularVelocitySmoothingThreshold) / FMath::Max(AngularVelocity.Length(), UE_SMALL_NUMBER);
+					const float ScalarVelocitySmoothingTime = (ClampedMaxPoseSmoothingTime * ScalarVelocitySmoothingThreshold) / FMath::Max(ScalarVelocity.Length(), UE_SMALL_NUMBER);
+
+					const float LinearVelocityFactor = 1.0f - FMath::InvExpApprox(DeltaTime / FMath::Clamp(LinearVelocitySmoothingTime, ClampedMinPoseSmoothingTime, ClampedMaxPoseSmoothingTime));
+					const float AngularVelocityFactor = 1.0f - FMath::InvExpApprox(DeltaTime / FMath::Clamp(AngularVelocitySmoothingTime, ClampedMinPoseSmoothingTime, ClampedMaxPoseSmoothingTime));
+					const float ScalarVelocityFactor = 1.0f - FMath::InvExpApprox(DeltaTime / FMath::Clamp(ScalarVelocitySmoothingTime, ClampedMinPoseSmoothingTime, ClampedMaxPoseSmoothingTime));
+
+					InOutPose.BoneLocations[FrameIdx][BoneIdx] = FMath::Lerp(BoneLocationIntegrated, InTargetPose.BoneLocations[FrameIdx][BoneIdx], LinearVelocityFactor);
+					InOutPose.BoneRotations[FrameIdx][BoneIdx] = FQuat4f::Slerp(BoneRotationIntegrated, InTargetPose.BoneRotations[FrameIdx][BoneIdx], AngularVelocityFactor);
+					InOutPose.BoneScales[FrameIdx][BoneIdx] = Math::VectorEerp(BoneScaleIntegrated, InTargetPose.BoneScales[FrameIdx][BoneIdx], ScalarVelocityFactor);
+					InOutPose.BoneLinearVelocities[FrameIdx][BoneIdx] = InTargetPose.BoneLinearVelocities[FrameIdx][BoneIdx];
+					InOutPose.BoneAngularVelocities[FrameIdx][BoneIdx] = InTargetPose.BoneAngularVelocities[FrameIdx][BoneIdx];
+					InOutPose.BoneScalarVelocities[FrameIdx][BoneIdx] = InTargetPose.BoneScalarVelocities[FrameIdx][BoneIdx];
+				}
+			}
+		}
+
+		void DampFramesInplaceUsingAdaptiveVelocityIntegrationMix(
+			const FPoseLocalBoneDataView& InOutPoses,
+			const float DeltaTime,
+			const float MinPoseSmoothingTime,
+			const float MaxPoseSmoothingTime,
+			const float LinearVelocitySmoothingThreshold,
+			const float AngularVelocitySmoothingThreshold,
+			const float ScalarVelocitySmoothingThreshold)
+		{
+			TRACE_CPUPROFILER_EVENT_SCOPE(AnimDatabase::PoseData::DampFramesInplaceUsingAdaptiveVelocityIntegrationMix);
+
+			const float ClampedMinPoseSmoothingTime = FMath::Max(MinPoseSmoothingTime, UE_SMALL_NUMBER);
+			const float ClampedMaxPoseSmoothingTime = FMath::Max(MaxPoseSmoothingTime, UE_SMALL_NUMBER);
+
+			const int32 FrameNum = InOutPoses.GetFrameNum();
+			const int32 BoneNum = InOutPoses.GetBoneNum();
+
+			for (int32 FrameIdx = 1; FrameIdx < FrameNum; FrameIdx++)
+			{
+				for (int32 BoneIdx = 0; BoneIdx < BoneNum; BoneIdx++)
+				{
+					// Doing the integration with a 50/50 velocity mix between start and end pose is more "implicit"
+					const FVector3f LinearVelocity = FMath::Lerp(InOutPoses.BoneLinearVelocities[FrameIdx - 1][BoneIdx], InOutPoses.BoneLinearVelocities[FrameIdx][BoneIdx], 0.5f);
+					const FVector3f AngularVelocity = FMath::Lerp(InOutPoses.BoneAngularVelocities[FrameIdx - 1][BoneIdx], InOutPoses.BoneAngularVelocities[FrameIdx][BoneIdx], 0.5f);
+					const FVector3f ScalarVelocity = FMath::Lerp(InOutPoses.BoneScalarVelocities[FrameIdx - 1][BoneIdx], InOutPoses.BoneScalarVelocities[FrameIdx][BoneIdx], 0.5f);
+
+					const FVector3f BoneLocationIntegrated = DeltaTime * LinearVelocity + InOutPoses.BoneLocations[FrameIdx - 1][BoneIdx];
+					const FQuat4f BoneRotationIntegrated = FQuat4f::MakeFromRotationVector(DeltaTime * AngularVelocity) * InOutPoses.BoneRotations[FrameIdx - 1][BoneIdx];
+					const FVector3f BoneScaleIntegrated = Math::VectorExpSafe(DeltaTime * ScalarVelocity) * InOutPoses.BoneScales[FrameIdx - 1][BoneIdx];
+
+					const float LinearVelocitySmoothingTime = (ClampedMaxPoseSmoothingTime * LinearVelocitySmoothingThreshold) / FMath::Max(LinearVelocity.Length(), UE_SMALL_NUMBER);
+					const float AngularVelocitySmoothingTime = (ClampedMaxPoseSmoothingTime * AngularVelocitySmoothingThreshold) / FMath::Max(AngularVelocity.Length(), UE_SMALL_NUMBER);
+					const float ScalarVelocitySmoothingTime = (ClampedMaxPoseSmoothingTime * ScalarVelocitySmoothingThreshold) / FMath::Max(ScalarVelocity.Length(), UE_SMALL_NUMBER);
+
+					const float LinearVelocityFactor = 1.0f - FMath::InvExpApprox(DeltaTime / FMath::Clamp(LinearVelocitySmoothingTime, ClampedMinPoseSmoothingTime, ClampedMaxPoseSmoothingTime));
+					const float AngularVelocityFactor = 1.0f - FMath::InvExpApprox(DeltaTime / FMath::Clamp(AngularVelocitySmoothingTime, ClampedMinPoseSmoothingTime, ClampedMaxPoseSmoothingTime));
+					const float ScalarVelocityFactor = 1.0f - FMath::InvExpApprox(DeltaTime / FMath::Clamp(ScalarVelocitySmoothingTime, ClampedMinPoseSmoothingTime, ClampedMaxPoseSmoothingTime));
+
+					InOutPoses.BoneLocations[FrameIdx][BoneIdx] = FMath::Lerp(BoneLocationIntegrated, InOutPoses.BoneLocations[FrameIdx][BoneIdx], LinearVelocityFactor);
+					InOutPoses.BoneRotations[FrameIdx][BoneIdx] = FQuat4f::Slerp(BoneRotationIntegrated, InOutPoses.BoneRotations[FrameIdx][BoneIdx], AngularVelocityFactor);
+					InOutPoses.BoneScales[FrameIdx][BoneIdx] = Math::VectorEerp(BoneScaleIntegrated, InOutPoses.BoneScales[FrameIdx][BoneIdx], ScalarVelocityFactor);
+				}
+			}
+		}
+
+		namespace ContactPinning
+		{
+			FContactStateConstView FContactStateConstView::Slice(const int32 FrameStart, const int32 FrameNum) const
+			{
+				return
+				{
+					Locked.Slice(FrameStart, FrameNum),
+					TimeSinceTransition.Slice(FrameStart, FrameNum),
+					Position.Slice(FrameStart, FrameNum),
+					Velocity.Slice(FrameStart, FrameNum),
+					Point.Slice(FrameStart, FrameNum),
+					OffsetPosition.Slice(FrameStart, FrameNum),
+					OffsetVelocity.Slice(FrameStart, FrameNum)
+				};
+			}
+
+			FContactStateView FContactStateView::Slice(const int32 FrameStart, const int32 FrameNum) const
+			{
+				return
+				{
+					Locked.Slice(FrameStart, FrameNum),
+					TimeSinceTransition.Slice(FrameStart, FrameNum),
+					Position.Slice(FrameStart, FrameNum),
+					Velocity.Slice(FrameStart, FrameNum),
+					Point.Slice(FrameStart, FrameNum),
+					OffsetPosition.Slice(FrameStart, FrameNum),
+					OffsetVelocity.Slice(FrameStart, FrameNum)
+				};
+			}
+
+			void FContactState::Resize(const int32 FrameNum, const int32 BoneNum)
+			{
+				Locked.SetNumUninitialized({ FrameNum, BoneNum });
+				TimeSinceTransition.SetNumUninitialized({ FrameNum, BoneNum });
+				Position.SetNumUninitialized({ FrameNum, BoneNum });
+				Velocity.SetNumUninitialized({ FrameNum, BoneNum });
+				Point.SetNumUninitialized({ FrameNum, BoneNum });
+				OffsetPosition.SetNumUninitialized({ FrameNum, BoneNum });
+				OffsetVelocity.SetNumUninitialized({ FrameNum, BoneNum });
+			}
+
+			bool FContactState::IsEmpty() const { return Locked.IsEmpty(); }
+
+			void FContactState::Empty()
+			{
+				Locked.Empty();
+				TimeSinceTransition.Empty();
+				Position.Empty();
+				Velocity.Empty();
+				Point.Empty();
+				OffsetPosition.Empty();
+				OffsetVelocity.Empty();
+			}
+
+			FContactStateView FContactState::View()
+			{
+				return
+				{
+					Locked,
+					TimeSinceTransition,
+					Position,
+					Velocity,
+					Point,
+					OffsetPosition,
+					OffsetVelocity
+				};
+			}
+
+			FContactStateConstView FContactState::ConstView() const
+			{
+				return
+				{
+					Locked,
+					TimeSinceTransition,
+					Position,
+					Velocity,
+					Point,
+					OffsetPosition,
+					OffsetVelocity
+				};
+			}
+
+			FContactStateView FContactState::Slice(const int32 FrameStart, const int32 FrameNum)
+			{
+				return
+				{
+					Locked.Slice(FrameStart, FrameNum),
+					TimeSinceTransition.Slice(FrameStart, FrameNum),
+					Position.Slice(FrameStart, FrameNum),
+					Velocity.Slice(FrameStart, FrameNum),
+					Point.Slice(FrameStart, FrameNum),
+					OffsetPosition.Slice(FrameStart, FrameNum),
+					OffsetVelocity.Slice(FrameStart, FrameNum)
+				};
+			}
+
+			FContactStateConstView FContactState::ConstSlice(const int32 FrameStart, const int32 FrameNum) const
+			{
+				return
+				{
+					Locked.Slice(FrameStart, FrameNum),
+					TimeSinceTransition.Slice(FrameStart, FrameNum),
+					Position.Slice(FrameStart, FrameNum),
+					Velocity.Slice(FrameStart, FrameNum),
+					Point.Slice(FrameStart, FrameNum),
+					OffsetPosition.Slice(FrameStart, FrameNum),
+					OffsetVelocity.Slice(FrameStart, FrameNum)
+				};
+			}
+
+			void ContactLookAtIK(
+				FQuat4f& InOutBoneRotation,
+				const FQuat4f GlobalParentRotation,
+				const FQuat4f GlobalRotation,
+				const FVector GlobalPosition,
+				const FVector ChildPosition,
+				const FVector TargetPosition,
+				const float Eps)
+			{
+				const FVector3f CurrDir = ((FVector3f)(ChildPosition - GlobalPosition)).GetSafeNormal(Eps, FVector3f::ForwardVector);
+				const FVector3f TargDir = ((FVector3f)(TargetPosition - GlobalPosition)).GetSafeNormal(Eps, FVector3f::ForwardVector);
+
+				if (FMath::Abs(1.0f - CurrDir.Dot(TargDir)) > Eps)
+				{
+					InOutBoneRotation = GlobalParentRotation.Inverse() * (FQuat4f::FindBetweenNormals(CurrDir, TargDir) * GlobalRotation);
+				}
+			}
+
+			void ContactTwoBoneIK(
+				FQuat4f& BoneRootLocalRotation,
+				FQuat4f& BoneMidLocalRotation,
+				const FVector BoneRoot,
+				const FVector BoneMid,
+				const FVector BoneEnd,
+				const FVector Target,
+				const FVector3f RotationAxis,
+				const FQuat4f BoneRootGlobalRotation,
+				const FQuat4f BoneMidGlobalRotation,
+				const FQuat4f BoneParentGlobalRotation,
+				const float MaxLengthBuffer)
+			{
+				const float MaxExtension = FVector::Distance(BoneRoot, BoneMid) + FVector::Distance(BoneMid, BoneEnd) - MaxLengthBuffer;
+
+				const float TargetDistance = FVector::Distance(Target, BoneRoot);
+				FVector TargetClamp = Target;
+
+				if (TargetDistance > MaxExtension)
+				{
+					const float Saturation = (1.0f - FMath::InvExpApprox((TargetDistance - MaxExtension) / FMath::Max(MaxLengthBuffer, UE_SMALL_NUMBER)));
+					TargetClamp = BoneRoot + (MaxExtension + MaxLengthBuffer * Saturation) * (Target - BoneRoot).GetUnsafeNormal();
+				}
+
+				const FVector3f AxisDwn = ((FVector3f)(BoneEnd - BoneRoot)).GetSafeNormal(UE_SMALL_NUMBER, FVector3f::ForwardVector);
+				const FVector3f AxisFwd = RotationAxis.Cross(AxisDwn).GetSafeNormal(UE_SMALL_NUMBER, FVector3f::UpVector);
+				const FVector3f AxisRot = AxisFwd.Cross(AxisDwn).GetSafeNormal(UE_SMALL_NUMBER, FVector3f::UpVector);
+
+				const FVector A = BoneRoot;
+				const FVector B = BoneMid;
+				const FVector C = BoneEnd;
+				const FVector T = TargetClamp;
+
+				const float LAB = FVector::Distance(B, A);
+				const float LCB = FVector::Distance(B, C);
+				const float LAT = FVector::Distance(T, A);
+
+				const float ACAB0 = FMath::Acos(((FVector3f)(C - A)).GetSafeNormal(UE_SMALL_NUMBER, FVector3f::ForwardVector).Dot(((FVector3f)(B - A)).GetSafeNormal(UE_SMALL_NUMBER, FVector3f::ForwardVector)));
+				const float BABC0 = FMath::Acos(((FVector3f)(A - B)).GetSafeNormal(UE_SMALL_NUMBER, FVector3f::ForwardVector).Dot(((FVector3f)(C - B)).GetSafeNormal(UE_SMALL_NUMBER, FVector3f::ForwardVector)));
+
+				const float ACAB1 = FMath::Acos((LAB * LAB + LAT * LAT - LCB * LCB) / FMath::Max(2.0f * LAB * LAT, UE_SMALL_NUMBER));
+				const float BABC1 = FMath::Acos((LAB * LAB + LCB * LCB - LAT * LAT) / FMath::Max(2.0f * LAB * LCB, UE_SMALL_NUMBER));
+
+				const FQuat4f R0 = FQuat4f::MakeFromRotationVector((ACAB1 - ACAB0) * AxisRot);
+				const FQuat4f R1 = FQuat4f::MakeFromRotationVector((BABC1 - BABC0) * AxisRot);
+				const FQuat4f R2 = FQuat4f::FindBetween((FVector3f)(BoneEnd - BoneRoot), (FVector3f)(TargetClamp - BoneRoot));
+
+				BoneRootLocalRotation = BoneParentGlobalRotation.Inverse() * (R2 * (R0 * BoneRootGlobalRotation));
+				BoneMidLocalRotation = BoneRootGlobalRotation.Inverse() * (R1 * BoneMidGlobalRotation);
+			}
+
+			void ResetContactPinning(
+				const FContactStateView& OutContactState,
+				const FPoseGlobalBoneDataConstView& InPoseGlobalBoneData,
+				const TArrayView<const int32> ToeBoneIndices)
+			{
+				check(OutContactState.GetFrameNum() == InPoseGlobalBoneData.GetFrameNum());
+				check(OutContactState.GetBoneNum() == ToeBoneIndices.Num());
+
+				const int32 FrameNum = OutContactState.GetFrameNum();
+				const int32 BoneNum = OutContactState.GetBoneNum();
+
+				for (int32 FrameIdx = 0; FrameIdx < FrameNum; FrameIdx++)
+				{
+					for (int32 BoneIdx = 0; BoneIdx < BoneNum; BoneIdx++)
+					{
+						OutContactState.Locked[FrameIdx][BoneIdx] = false;
+						OutContactState.TimeSinceTransition[FrameIdx][BoneIdx] = 0.0f;
+						OutContactState.Position[FrameIdx][BoneIdx] = InPoseGlobalBoneData.BoneLocations[FrameIdx][ToeBoneIndices[BoneIdx]];
+						OutContactState.Velocity[FrameIdx][BoneIdx] = InPoseGlobalBoneData.BoneLinearVelocities[FrameIdx][ToeBoneIndices[BoneIdx]];
+						OutContactState.Point[FrameIdx][BoneIdx] = InPoseGlobalBoneData.BoneLocations[FrameIdx][ToeBoneIndices[BoneIdx]];
+						OutContactState.OffsetPosition[FrameIdx][BoneIdx] = FVector3f::ZeroVector;
+						OutContactState.OffsetVelocity[FrameIdx][BoneIdx] = FVector3f::ZeroVector;
+					}
+				}
+			}
+
+			void UpdateContactPinning(
+				const FContactStateView& InOutContactState,
+				const FPoseGlobalBoneDataView& InOutPoseGlobalBoneData,
+				const FPoseLocalBoneDataView& InOutPoseLocalBoneData,
+				const FPoseRootDataConstView& InPoseRootData,
+				const FPoseAttributeDataConstView& InPoseAttributeData,
+				const TArrayView<const int32> InBoneParents,
+				const TArrayView<const int32> ContactStateAttributeIndices,
+				const TArrayView<const int32> PelvisBoneIndices,
+				const TArrayView<const int32> HipBoneIndices,
+				const TArrayView<const int32> KneeBoneIndices,
+				const TArrayView<const int32> AnkleBoneIndices,
+				const TArrayView<const int32> ToeBoneIndices,
+				const TArrayView<const FVector3f> KneeSideVectors,
+				const float DeltaTime,
+				const FContactPinningSettings& Settings)
+			{
+				check(InOutContactState.GetFrameNum() == InOutPoseGlobalBoneData.GetFrameNum());
+				check(InOutContactState.GetFrameNum() == InOutPoseLocalBoneData.GetFrameNum());
+				check(InOutContactState.GetFrameNum() == InPoseRootData.GetFrameNum());
+				check(InOutContactState.GetFrameNum() == InPoseAttributeData.GetFrameNum());
+				check(InOutContactState.GetBoneNum() == ContactStateAttributeIndices.Num());
+				check(InOutContactState.GetBoneNum() == PelvisBoneIndices.Num());
+				check(InOutContactState.GetBoneNum() == HipBoneIndices.Num());
+				check(InOutContactState.GetBoneNum() == KneeBoneIndices.Num());
+				check(InOutContactState.GetBoneNum() == AnkleBoneIndices.Num());
+				check(InOutContactState.GetBoneNum() == ToeBoneIndices.Num());
+				check(InOutContactState.GetBoneNum() == KneeSideVectors.Num());
+
+				const int32 FrameNum = InOutContactState.GetFrameNum();
+				const int32 BoneNum = InOutContactState.GetBoneNum();
+
+				TArray<int32, TInlineAllocator<32>> RequiredBones;
+				TArray<int32, TInlineAllocator<32>> Lhs;
+				TArray<int32, TInlineAllocator<32>> Rhs;
+				for (int32 BoneIdx = 0; BoneIdx < BoneNum; BoneIdx++)
+				{
+					Lhs = RequiredBones;
+					Rhs.SetNumUninitialized(Math::BoneAscendantsInclusiveNum(ToeBoneIndices[BoneIdx], InBoneParents));
+					Math::BoneAscendantsInclusive(Rhs, ToeBoneIndices[BoneIdx], InBoneParents);
+					RequiredBones.SetNumUninitialized(Math::BoneUnionNum(Lhs, Rhs));
+					Math::BoneUnion(RequiredBones, Lhs, Rhs);
+				}
+
+				PoseData::ForwardKinematicsPartial(
+					InOutPoseGlobalBoneData,
+					InOutPoseLocalBoneData.ConstView(),
+					InPoseRootData,
+					InBoneParents,
+					UE::Learning::FIndexSet(0, InBoneParents.Num()),
+					RequiredBones);
+
+				for (int32 FrameIdx = 0; FrameIdx < FrameNum; FrameIdx++)
+				{
+					for (int32 BoneIdx = 0; BoneIdx < BoneNum; BoneIdx++)
+					{
+						const int32 AttributeIdx = ContactStateAttributeIndices[BoneIdx];
+
+						const FVector ContactPosition = InOutPoseGlobalBoneData.BoneLocations[FrameIdx][ToeBoneIndices[BoneIdx]];
+						const FVector3f ContactVelocity = InOutPoseGlobalBoneData.BoneLinearVelocities[FrameIdx][ToeBoneIndices[BoneIdx]];
+						const bool bContactState =
+							(AttributeIdx >= 0 && AttributeIdx < InPoseAttributeData.GetAttributeNum()) &&
+							InPoseAttributeData.GetAttributeType(AttributeIdx) == EAnimDatabaseAttributeType::Float &&
+							InPoseAttributeData.GetAttributeActive(FrameIdx, AttributeIdx) &&
+							InPoseAttributeData.GetFloat(FrameIdx, AttributeIdx) > 0.5f;
+
+						UE::AnimDatabase::Math::LocationInertializeCubicUpdate(
+							InOutContactState.Position[FrameIdx][BoneIdx],
+							InOutContactState.Velocity[FrameIdx][BoneIdx],
+							InOutContactState.TimeSinceTransition[FrameIdx][BoneIdx],
+							InOutContactState.Locked[FrameIdx][BoneIdx] ? InOutContactState.Point[FrameIdx][BoneIdx] : ContactPosition,
+							InOutContactState.Locked[FrameIdx][BoneIdx] ? FVector3f::ZeroVector : ContactVelocity,
+							InOutContactState.OffsetPosition[FrameIdx][BoneIdx],
+							InOutContactState.OffsetVelocity[FrameIdx][BoneIdx],
+							DeltaTime,
+							Settings.BlendTime);
+
+						const float LockDistanceSquared = (InOutContactState.Position[FrameIdx][BoneIdx] - ContactPosition).SquaredLength();
+
+						if (!InOutContactState.Locked[FrameIdx][BoneIdx] && bContactState && LockDistanceSquared < FMath::Square(Settings.LockRadius))
+						{
+							InOutContactState.Locked[FrameIdx][BoneIdx] = true;
+							InOutContactState.Point[FrameIdx][BoneIdx] = ContactPosition;
+							InOutContactState.Point[FrameIdx][BoneIdx].Z = Settings.ContactHeight;
+
+							UE::AnimDatabase::Math::LocationInertializeCubicTransition(
+								InOutContactState.OffsetPosition[FrameIdx][BoneIdx],
+								InOutContactState.OffsetVelocity[FrameIdx][BoneIdx],
+								InOutContactState.TimeSinceTransition[FrameIdx][BoneIdx],
+								ContactPosition,
+								ContactVelocity,
+								InOutContactState.Point[FrameIdx][BoneIdx],
+								FVector3f::ZeroVector,
+								Settings.BlendTime);
+						}
+						else if (InOutContactState.Locked[FrameIdx][BoneIdx] && (!bContactState || LockDistanceSquared > FMath::Square(Settings.UnlockRadius)))
+						{
+							InOutContactState.Locked[FrameIdx][BoneIdx] = false;
+
+							UE::AnimDatabase::Math::LocationInertializeCubicTransition(
+								InOutContactState.OffsetPosition[FrameIdx][BoneIdx],
+								InOutContactState.OffsetVelocity[FrameIdx][BoneIdx],
+								InOutContactState.TimeSinceTransition[FrameIdx][BoneIdx],
+								InOutContactState.Point[FrameIdx][BoneIdx],
+								FVector3f::ZeroVector,
+								ContactPosition,
+								ContactVelocity,
+								Settings.BlendTime);
+						}
+
+						FVector ToePosition = InOutContactState.Position[FrameIdx][BoneIdx];
+						ToePosition.Z = FMath::Max(ToePosition.Z, Settings.MinToeHeight);
+
+						FVector AnklePosition = ToePosition + (InOutPoseGlobalBoneData.BoneLocations[FrameIdx][AnkleBoneIndices[BoneIdx]] - InOutPoseGlobalBoneData.BoneLocations[FrameIdx][ToeBoneIndices[BoneIdx]]);
+						AnklePosition.Z = FMath::Max(AnklePosition.Z, Settings.MinAnkleHeight);
+
+						ContactTwoBoneIK(
+							InOutPoseLocalBoneData.BoneRotations[FrameIdx][HipBoneIndices[BoneIdx]],
+							InOutPoseLocalBoneData.BoneRotations[FrameIdx][KneeBoneIndices[BoneIdx]],
+							InOutPoseGlobalBoneData.BoneLocations[FrameIdx][HipBoneIndices[BoneIdx]],
+							InOutPoseGlobalBoneData.BoneLocations[FrameIdx][KneeBoneIndices[BoneIdx]],
+							InOutPoseGlobalBoneData.BoneLocations[FrameIdx][AnkleBoneIndices[BoneIdx]],
+							AnklePosition,
+							InOutPoseGlobalBoneData.BoneRotations[FrameIdx][KneeBoneIndices[BoneIdx]].RotateVector(KneeSideVectors[BoneIdx]),
+							InOutPoseGlobalBoneData.BoneRotations[FrameIdx][HipBoneIndices[BoneIdx]],
+							InOutPoseGlobalBoneData.BoneRotations[FrameIdx][KneeBoneIndices[BoneIdx]],
+							InOutPoseGlobalBoneData.BoneRotations[FrameIdx][PelvisBoneIndices[BoneIdx]],
+							Settings.HyperExtensionLimit);
+					}
+				}
+
+				PoseData::ForwardKinematicsPartial(
+					InOutPoseGlobalBoneData,
+					InOutPoseLocalBoneData.ConstView(),
+					InPoseRootData,
+					InBoneParents,
+					UE::Learning::FIndexSet(0, InBoneParents.Num()),
+					RequiredBones);
+
+				for (int32 FrameIdx = 0; FrameIdx < FrameNum; FrameIdx++)
+				{
+					for (int32 BoneIdx = 0; BoneIdx < BoneNum; BoneIdx++)
+					{
+						const FVector ToePosition = InOutContactState.Position[FrameIdx][BoneIdx];
+
+						ContactLookAtIK(
+							InOutPoseLocalBoneData.BoneRotations[FrameIdx][AnkleBoneIndices[BoneIdx]],
+							InOutPoseGlobalBoneData.BoneRotations[FrameIdx][KneeBoneIndices[BoneIdx]],
+							InOutPoseGlobalBoneData.BoneRotations[FrameIdx][AnkleBoneIndices[BoneIdx]],
+							InOutPoseGlobalBoneData.BoneLocations[FrameIdx][AnkleBoneIndices[BoneIdx]],
+							InOutPoseGlobalBoneData.BoneLocations[FrameIdx][ToeBoneIndices[BoneIdx]],
+							ToePosition);
+					}
+				}
+
+				PoseData::ForwardKinematicsPartial(
+					InOutPoseGlobalBoneData,
+					InOutPoseLocalBoneData.ConstView(),
+					InPoseRootData,
+					InBoneParents,
+					UE::Learning::FIndexSet(0, InBoneParents.Num()),
+					RequiredBones);
+
+				for (int32 FrameIdx = 0; FrameIdx < FrameNum; FrameIdx++)
+				{
+					for (int32 BoneIdx = 0; BoneIdx < BoneNum; BoneIdx++)
+					{
+						const FVector ToeEndCurr =
+							(FVector)InOutPoseGlobalBoneData.BoneRotations[FrameIdx][ToeBoneIndices[BoneIdx]].RotateVector(FVector3f(-Settings.ToeLength, 0.0f, 0.0f)) +
+							InOutPoseGlobalBoneData.BoneLocations[FrameIdx][ToeBoneIndices[BoneIdx]];
+
+						FVector ToeEndTarg = ToeEndCurr;
+						ToeEndTarg.Z = FMath::Max(ToeEndTarg.Z, Settings.MinToeHeight);
+
+						ContactLookAtIK(
+							InOutPoseLocalBoneData.BoneRotations[FrameIdx][ToeBoneIndices[BoneIdx]],
+							InOutPoseGlobalBoneData.BoneRotations[FrameIdx][AnkleBoneIndices[BoneIdx]],
+							InOutPoseGlobalBoneData.BoneRotations[FrameIdx][ToeBoneIndices[BoneIdx]],
+							InOutPoseGlobalBoneData.BoneLocations[FrameIdx][ToeBoneIndices[BoneIdx]],
+							ToeEndCurr,
+							ToeEndTarg);
+					}
+				}
+			}
+
 		}
 	}
 }
diff --git a/Engine/Plugins/Experimental/Animation/AnimDatabase/Source/AnimDatabase/Public/AnimDatabasePose.h b/Engine/Plugins/Experimental/Animation/AnimDatabase/Source/AnimDatabase/Public/AnimDatabasePose.h
--- a/Engine/Plugins/Experimental/Animation/AnimDatabase/Source/AnimDatabase/Public/AnimDatabasePose.h
+++ b/Engine/Plugins/Experimental/Animation/AnimDatabase/Source/AnimDatabase/Public/AnimDatabasePose.h
@@ -416,7 +416,7 @@ namespace UE::AnimDatabase
 		 * @param InPoseLocalBoneData		Input Local Bone Transforms of shape (FrameNum, BoneIndicesNum)
 		 * @param InPoseRootData			Input Root Transforms of shape (FrameNum)
 		 * @param BoneParents				Input Bone Parents Array of shape (BoneNum)
-		 * @param BoneIndices				Input Bone Parents Array of shape (BoneIndicesNum)
+		 * @param BoneIndices				Input Bone Indices Array of shape (BoneIndicesNum)
 		 */
 		UE_API void ForwardKinematics(
 			const FPoseGlobalBoneDataView& OutPoseGlobalBoneData,
@@ -439,6 +439,28 @@ namespace UE::AnimDatabase
 			const FPoseRootDataConstView& InPoseRootData,
 			const TLearningArrayView<1, const int32> BoneParents);
 
+		/**
+		 * Compute Forward Kinematics on a subset of bones to produce global transforms from local transforms.
+		 *
+		 * The OutPoseGlobalBoneData and InPoseLocalBoneData parameters are expected to be allocated for just the subset of bones being used, which
+		 * is then provided by BoneIndices. The BoneParents array however should still refer to bone indices in the full reference skeleton. All
+		 * required parent bones must be present in BoneIndices.
+		 *
+		 * @param OutPoseGlobalBoneData		Output Global Bone Transforms of shape (FrameNum, BoneIndicesNum)
+		 * @param InPoseLocalBoneData		Input Local Bone Transforms of shape (FrameNum, BoneIndicesNum)
+		 * @param InPoseRootData			Input Root Transforms of shape (FrameNum)
+		 * @param BoneParents				Input Bone Parents Array of shape (BoneNum)
+		 * @param BoneIndices				Input Bone Indices Array of shape (BoneIndicesNum)
+		 * @param BonesToUpdate				Which bones to compute the global transforms for
+		 */
+		UE_API void ForwardKinematicsPartial(
+			const FPoseGlobalBoneDataView& OutPoseGlobalBoneData,
+			const FPoseLocalBoneDataConstView& InPoseLocalBoneData,
+			const FPoseRootDataConstView& InPoseRootData,
+			const TLearningArrayView<1, const int32> BoneParents,
+			const UE::Learning::FIndexSet BoneIndices,
+			const UE::Learning::FIndexSet BonesToUpdate);
+
 		/**
 		 * Computes the bone indices which are different from the reference pose transforms in the given pose data
 		 * 
@@ -618,6 +640,264 @@ namespace UE::AnimDatabase
 
 		/** Recomputes Velocities using finite difference assuming the first dimension is frames */
 		UE_API void RecomputeVelocities(const FPoseRootDataView& InOutRootData, const float DeltaTime);
+
+		/** Teleport some pose data to a new location and rotation. Adjusts the velocities appropriately. */
+		UE_API void Teleport(
+			const FPoseRootDataView& InOutRootData,
+			const TLearningArrayView<1, const FVector> RootLocations,
+			const TLearningArrayView<1, const FQuat4f> RootRotations);
+
+		/** Blend root data using a given alpha value */
+		UE_API void BlendRoot(
+			const FPoseRootDataView& OutPoseData,
+			const FPoseRootDataConstView& InPoseA,
+			const FPoseRootDataConstView& InPoseB,
+			const float Alpha);
+
+		/** Integrates the root location and rotation using the velocities. */
+		UE_API void IntegrateRoot(const FPoseRootDataView& InOutRootData, const float DeltaTime);
+
+		/** Blend pose data of two poses, ensuring to blend via the given set of rotations. Updates the blend via rotations. */
+		UE_API void BlendVia(
+			const FPoseLocalBoneDataView& OutPoseData,
+			const TLearningArrayView<2, FQuat4f> InOutViaRotations,
+			const FPoseLocalBoneDataConstView& InPoseA,
+			const FPoseLocalBoneDataConstView& InPoseB,
+			const float Alpha);
+
+		/** Fit extrapolation half-lives given source and target pose data */
+		UE_API void FitExtrapolationHalfLives(
+			const TLearningArrayView<2, FVector3f> OutBoneLocationHalfLives,
+			const TLearningArrayView<2, FVector3f> OutBoneRotationHalfLives,
+			const TLearningArrayView<2, FVector3f> OutBoneScaleHalfLives,
+			const FPoseLocalBoneDataConstView& InPoseSource,
+			const FPoseLocalBoneDataConstView& InPoseTarget,
+			const float ExtrapolationHalfLife = 1.0f,
+			const float ExtrapolationHalfLifeMin = 0.05f,
+			const float ExtrapolationHalfLifeMax = 1.0f,
+			const float MaxLinearVelocity = 500.0f,
+			const float MaxAngularVelocity = UE_TWO_PI,
+			const float MaxScalarVelocity = 4.0f);
+
+		/** Extrapolate the root data forward given a half-life */
+		UE_API void ExtrapolateRoot(
+			const FPoseRootDataView& OutPoseData,
+			const FPoseRootDataConstView& InPoseData,
+			const float Time,
+			const float RootHalfLife = 1.0f);
+
+		/** Extrapolate pose data forward in time given sets of half-lives */
+		UE_API void Extrapolate(
+			const FPoseLocalBoneDataView& OutPoseData,
+			const FPoseLocalBoneDataConstView& InPoseData,
+			const float Time,
+			const TLearningArrayView<2, const FVector3f> BoneLocationHalfLives,
+			const TLearningArrayView<2, const FVector3f> BoneRotationHalfLives,
+			const TLearningArrayView<2, const FVector3f> BoneScaleHalfLives);
+
+		/** Fit inertialization offsets to the provided root data */
+		UE_API void FitInertializationOffsetsRoot(
+			const FPoseRootDataView& OutOffsetPose,
+			const FPoseRootDataConstView& InPoseSource,
+			const FPoseRootDataConstView& InPoseTarget);
+
+		/** Apply inertialization offsets to the given root data */
+		UE_API void ApplyInertializationOffsetsRoot(
+			const FPoseRootDataView& OutFinalPose,
+			const FPoseRootDataConstView& InOffsetPose,
+			const FPoseRootDataConstView& InPoseTarget,
+			const float BlendTime,
+			const float BlendDuration);
+
+		/** Fit inertialization offsets to the provided pose data */
+		UE_API void FitInertializationOffsets(
+			const FPoseLocalBoneDataView& OutOffsetPose,
+			const FPoseLocalBoneDataConstView& InPoseSource,
+			const FPoseLocalBoneDataConstView& InPoseTarget);
+
+		/** Apply inertialization offsets to the given pose data */
+		UE_API void ApplyInertializationOffsets(
+			const FPoseLocalBoneDataView& OutFinalPose,
+			const FPoseLocalBoneDataConstView& InOffsetPose,
+			const FPoseLocalBoneDataConstView& InPoseTarget,
+			const float BlendTime,
+			const float BlendDuration);
+
+		/** Fit inertialization offsets to the provided attribute data */
+		UE_API void FitInertializationOffsetsAttributes(
+			const FPoseAttributeDataView& OutOffsetPose,
+			const FPoseAttributeDataConstView& InPoseSource,
+			const FPoseAttributeDataConstView& InPoseTarget);
+
+		/** Apply inertialization offsets to the given attribute data */
+		UE_API void ApplyInertializationOffsetsAttributes(
+			const FPoseAttributeDataView& OutFinalPose,
+			const FPoseAttributeDataConstView& InOffsetPose,
+			const FPoseAttributeDataConstView& InPoseTarget,
+			const float BlendTime,
+			const float BlendDuration);
+
+		/** Update root data in-place using a mix of integrating the velocity and the next target pose */
+		UE_API void UpdateInplaceUsingVelocityIntegrationMix(
+			const FPoseRootDataView& InOutPose,
+			const FPoseRootDataConstView& InTargetPose,
+			const float DeltaTime,
+			const float VelocityAlpha = 0.25f);
+
+		/** Update pose data in-place using a mix of integrating the velocity and the next target pose */
+		UE_API void UpdateInplaceUsingVelocityIntegrationMix(
+			const FPoseLocalBoneDataView& InOutPose,
+			const FPoseLocalBoneDataConstView& InTargetPose,
+			const float DeltaTime,
+			const float VelocityAlpha = 0.25f);
+
+		/** Update pose data in-place using a mix of integrating the velocity and the next target pose */
+		UE_API void DampInplaceUsingVelocityIntegrationMix(
+			const FPoseLocalBoneDataView& InOutPose,
+			const FPoseLocalBoneDataConstView& InTargetPose,
+			const float DeltaTime,
+			const float PoseSmoothingTime = 0.2f);
+
+		/** Update pose data in-place using an adaptive mix of integrating the velocity and the next target pose based on the absolute velocity magnitude */
+		UE_API void DampInplaceUsingAdaptiveVelocityIntegrationMix(
+			const FPoseLocalBoneDataView& InOutPose,
+			const FPoseLocalBoneDataConstView& InTargetPose,
+			const float DeltaTime,
+			const float MinPoseSmoothingTime = 0.0f,
+			const float MaxPoseSmoothingTime = 0.25f,
+			const float LinearVelocitySmoothingThreshold = 5.0f,
+			const float AngularVelocitySmoothingThreshold = UE_PI / 8.0f,
+			const float ScalarVelocitySmoothingThreshold = 0.5f);
+
+		/** Update pose data representing a sequence of frames in-place using an adaptive mix of integrating the velocity and the next target pose based on the absolute velocity magnitude */
+		UE_API void DampFramesInplaceUsingAdaptiveVelocityIntegrationMix(
+			const FPoseLocalBoneDataView& InOutPoses,
+			const float DeltaTime,
+			const float MinPoseSmoothingTime = 0.0f,
+			const float MaxPoseSmoothingTime = 0.25f,
+			const float LinearVelocitySmoothingThreshold = 5.0f,
+			const float AngularVelocitySmoothingThreshold = UE_PI / 8.0f,
+			const float ScalarVelocitySmoothingThreshold = 0.5f);
+
+		/** Helper functions for performing contact pinning */
+		namespace ContactPinning
+		{
+			/** Const View of Contact States */
+			struct FContactStateConstView
+			{
+				UE_API int32 GetFrameNum() const { return Locked.Num<0>(); }
+				UE_API int32 GetBoneNum() const { return Locked.Num<1>(); }
+				UE_API FContactStateConstView Slice(const int32 FrameStart, const int32 FrameNum) const;
+
+				TLearningArrayView<2, const bool> Locked;
+				TLearningArrayView<2, const float> TimeSinceTransition;
+				TLearningArrayView<2, const FVector> Position;
+				TLearningArrayView<2, const FVector3f> Velocity;
+				TLearningArrayView<2, const FVector> Point;
+				TLearningArrayView<2, const FVector3f> OffsetPosition;
+				TLearningArrayView<2, const FVector3f> OffsetVelocity;
+			};
+
+			/** Non-Const View of Contact States */
+			struct FContactStateView
+			{
+				UE_API int32 GetFrameNum() const { return Locked.Num<0>(); }
+				UE_API int32 GetBoneNum() const { return Locked.Num<1>(); }
+				UE_API FContactStateView Slice(const int32 FrameStart, const int32 FrameNum) const;
+
+				TLearningArrayView<2, bool> Locked;
+				TLearningArrayView<2, float> TimeSinceTransition;
+				TLearningArrayView<2, FVector> Position;
+				TLearningArrayView<2, FVector3f> Velocity;
+				TLearningArrayView<2, FVector> Point;
+				TLearningArrayView<2, FVector3f> OffsetPosition;
+				TLearningArrayView<2, FVector3f> OffsetVelocity;
+			};
+
+			/** Contact State data */
+			struct FContactState
+			{
+				UE_API int32 GetFrameNum() const { return Locked.Num<0>(); }
+				UE_API int32 GetBoneNum() const { return Locked.Num<1>(); }
+				UE_API void Resize(const int32 FrameNum, const int32 BoneNum);
+				UE_API bool IsEmpty() const;
+				UE_API void Empty();
+
+				UE_API FContactStateView View();
+				UE_API FContactStateConstView ConstView() const;
+				UE_API FContactStateView Slice(const int32 FrameStart, const int32 FrameNum);
+				UE_API FContactStateConstView ConstSlice(const int32 FrameStart, const int32 FrameNum) const;
+
+				TLearningArray<2, bool> Locked;
+				TLearningArray<2, float> TimeSinceTransition;
+				TLearningArray<2, FVector> Position;
+				TLearningArray<2, FVector3f> Velocity;
+				TLearningArray<2, FVector> Point;
+				TLearningArray<2, FVector3f> OffsetPosition;
+				TLearningArray<2, FVector3f> OffsetVelocity;
+			};
+
+			/** Contact look-at IK helper function */
+			UE_API void ContactLookAtIK(
+				FQuat4f& InOutBoneRotation,
+				const FQuat4f GlobalParentRotation,
+				const FQuat4f GlobalRotation,
+				const FVector GlobalPosition,
+				const FVector ChildPosition,
+				const FVector TargetPosition,
+				const float Eps = 1e-5f);
+
+			/** Contact two-bone IK helper function */
+			UE_API void ContactTwoBoneIK(
+				FQuat4f& BoneRootLocalRotation,
+				FQuat4f& BoneMidLocalRotation,
+				const FVector BoneRoot,
+				const FVector BoneMid,
+				const FVector BoneEnd,
+				const FVector Target,
+				const FVector3f RotationAxis,
+				const FQuat4f BoneRootGlobalRotation,
+				const FQuat4f BoneMidGlobalRotation,
+				const FQuat4f BoneParentGlobalRotation,
+				const float MaxLengthBuffer = 1.0f);
+
+			/** Resets the contact pinning state given the pose data and toe indices */
+			UE_API void ResetContactPinning(
+				const FContactStateView& OutContactState,
+				const FPoseGlobalBoneDataConstView& InPoseGlobalBoneData,
+				const TArrayView<const int32> ToeBoneIndices);
+
+			/** Contact pinning settings */
+			struct FContactPinningSettings
+			{
+				float BlendTime = 0.25f;
+				float LockRadius = 5.0f;
+				float UnlockRadius = 20.0f;
+				float ContactHeight = 0.0f;
+				float MinToeHeight = 0.0f;
+				float MinAnkleHeight = 7.0f;
+				float ToeLength = 5.0f;
+				float HyperExtensionLimit = 0.5f;
+			};
+
+			/** Updates the contact pinning state using the provided pose data, bone indices, and settings */
+			UE_API void UpdateContactPinning(
+				const FContactStateView& InOutContactState,
+				const FPoseGlobalBoneDataView& InOutPoseGlobalBoneData,
+				const FPoseLocalBoneDataView& InOutPoseLocalBoneData,
+				const FPoseRootDataConstView& InPoseRootData,
+				const FPoseAttributeDataConstView& InPoseAttributeData,
+				const TArrayView<const int32> InBoneParents,
+				const TArrayView<const int32> ContactStateAttributeIndices,
+				const TArrayView<const int32> PelvisBoneIndices,
+				const TArrayView<const int32> HipBoneIndices,
+				const TArrayView<const int32> KneeBoneIndices,
+				const TArrayView<const int32> AnkleBoneIndices,
+				const TArrayView<const int32> ToeBoneIndices,
+				const TArrayView<const FVector3f> KneeSideVectors,
+				const float DeltaTime,
+				const FContactPinningSettings& Settings = FContactPinningSettings());
+		}
 	}
 
 	//--------------------------------------------------
diff --git a/Engine/Plugins/Experimental/Animation/AnimDatabase/Source/AnimDatabaseEditor/Private/AnimDatabaseEditorMode.cpp b/Engine/Plugins/Experimental/Animation/AnimDatabase/Source/AnimDatabaseEditor/Private/AnimDatabaseEditorMode.cpp
--- a/Engine/Plugins/Experimental/Animation/AnimDatabase/Source/AnimDatabaseEditor/Private/AnimDatabaseEditorMode.cpp
+++ b/Engine/Plugins/Experimental/Animation/AnimDatabase/Source/AnimDatabaseEditor/Private/AnimDatabaseEditorMode.cpp
@@ -547,7 +547,7 @@ namespace UE::AnimDatabase::Editor
 
 		for (int32 CharacterIdx = 0; CharacterIdx < CharacterNum; CharacterIdx++)
 		{
-			if (PreviewScene->GetPoseableMeshComponent(CharacterIdx)->GetSkinnedAsset() != nullptr)
+			if (PreviewScene->GetPoseableMeshComponent(CharacterIdx)->GetSkinnedAsset())
 			{
 				for (int32 BoneIdx = 0; BoneIdx < BoneNum; BoneIdx++)
 				{
@@ -556,11 +556,10 @@ namespace UE::AnimDatabase::Editor
 						(FVector)PoseState.PoseGlobalBoneData->BoneLocations[CharacterIdx][BoneIdx],
 						(FVector)PoseState.PoseGlobalBoneData->BoneScales[CharacterIdx][BoneIdx]);
 
-					const FName BoneName = Database->GetSkeleton()->GetReferenceSkeleton().GetBoneName(BoneIdx);
-					PreviewScene->GetPoseableMeshComponent(CharacterIdx)->SetBoneTransformByName(BoneName, BoneTransform, EBoneSpaces::WorldSpace);
+					PreviewScene->GetPoseableMeshComponent(CharacterIdx)->SetBoneTransformByName(Database->GetBoneName(BoneIdx), BoneTransform, EBoneSpaces::WorldSpace);
 				}
 
-				PreviewScene->GetPoseableMeshComponent(CharacterIdx)->RefreshBoneTransforms(nullptr);
+				PreviewScene->GetPoseableMeshComponent(CharacterIdx)->RefreshBoneTransforms();
 			}
 		}
 
