AnimDatabase: Moved a bunch of stuff to editor only

#Synced-CL 51119969

diff --git a/Engine/Plugins/Experimental/Animation/AnimDatabase/Source/AnimDatabase/Private/AnimDatabase.cpp b/Engine/Plugins/Experimental/Animation/AnimDatabase/Source/AnimDatabase/Private/AnimDatabase.cpp
--- a/Engine/Plugins/Experimental/Animation/AnimDatabase/Source/AnimDatabase/Private/AnimDatabase.cpp
+++ b/Engine/Plugins/Experimental/Animation/AnimDatabase/Source/AnimDatabase/Private/AnimDatabase.cpp
@@ -370,12 +370,6 @@
 		ContainerBoneIndices.Empty(bMirrored ? BoneNum * 4 + 1 : BoneNum + 1);
 		ContainerBoneIndices.AddUnique(0); // Seems important to always include root bone...
 
-		// Add desired bones
-		for (int32 BoneIdx = 0; BoneIdx < BoneNum; BoneIdx++)
-		{
-			ContainerBoneIndices.AddUnique(BoneIndices[BoneIdx]);
-		}
-
 		// To mirror we need to include both the parent bone index, the mirror bone index, plus the parent of the mirror bone index
 		if (bMirrored)
 		{
@@ -384,25 +378,38 @@
 
 			for (int32 BoneIdx = 0; BoneIdx < BoneNum; BoneIdx++)
 			{
-				const int32 ParentBoneIndex = AnimSequence->GetSkeleton()->GetReferenceSkeleton().GetParentIndex(BoneIndices[BoneIdx]);
-				if (ParentBoneIndex != INDEX_NONE)
+				// Add bone and ascendants
+				ContainerBoneIndices.AddUnique(BoneIndices[BoneIdx]);
+
+				int32 ParentBoneIndex = AnimSequence->GetSkeleton()->GetReferenceSkeleton().GetParentIndex(BoneIndices[BoneIdx]);
+				while (ParentBoneIndex != INDEX_NONE)
 				{
 					ContainerBoneIndices.AddUnique(ParentBoneIndex);
+					ParentBoneIndex = AnimSequence->GetSkeleton()->GetReferenceSkeleton().GetParentIndex(ParentBoneIndex);
 				}
 
+				// Add mirror bone and ascendants
 				const int32 MirrorBoneIndex = MirrorBoneIndexes[BoneIndices[BoneIdx]].GetInt();
 				if (MirrorBoneIndex != INDEX_NONE)
 				{
 					ContainerBoneIndices.AddUnique(MirrorBoneIndex);
 
-					const int32 MirrorParentBoneIndex = AnimSequence->GetSkeleton()->GetReferenceSkeleton().GetParentIndex(MirrorBoneIndex);
-					if (MirrorParentBoneIndex != INDEX_NONE)
+					int32 MirrorParentBoneIndex = AnimSequence->GetSkeleton()->GetReferenceSkeleton().GetParentIndex(MirrorBoneIndex);
+					while (MirrorParentBoneIndex != INDEX_NONE)
 					{
 						ContainerBoneIndices.AddUnique(MirrorParentBoneIndex);
+						MirrorParentBoneIndex = AnimSequence->GetSkeleton()->GetReferenceSkeleton().GetParentIndex(MirrorParentBoneIndex);
 					}
 				}
 			}
 		}
+		else
+		{
+			for (int32 BoneIdx = 0; BoneIdx < BoneNum; BoneIdx++)
+			{
+				ContainerBoneIndices.AddUnique(BoneIndices[BoneIdx]);
+			}
+		}
 
 		// Make sure bone indices are sorted
 
@@ -2260,6 +2267,8 @@
 #endif
 }
 
+#if WITH_EDITOR
+
 void UAnimDatabaseQuery::UpdateQueryEntries()
 {
 	QueryEntries.Reset();
@@ -2560,10 +2569,14 @@
 	const int32 RangeLength,
 	const FLinearColor& IdentifierColor) {}
 
+#endif
+
 UAnimDatabase::UAnimDatabase(const FObjectInitializer& ObjectInitializer) : Super(ObjectInitializer)
 {
+#if WITH_EDITOR
 	ViewportSettings = ObjectInitializer.CreateEditorOnlyDefaultSubobject<UAnimDatabaseViewportSettings>(this, TEXT("ViewportSettings"));
 	Query = ObjectInitializer.CreateEditorOnlyDefaultSubobject<UAnimDatabaseQuery>(this, TEXT("Query"));
+#endif
 }
 
 #if WITH_EDITOR
@@ -2920,7 +2933,14 @@
 
 FString UAnimDatabase::GetSequencePathString(const int32 SequenceIdx) const
 {
-	return GetAnimSequence(SequenceIdx)->GetPathName(this);
+	if (UAnimSequence* Sequence = GetAnimSequence(SequenceIdx))
+	{
+		return Sequence->GetPathName(this);
+	}
+	else
+	{
+		return FString();
+	}
 }
 
 int32 UAnimDatabase::GetBoneParent(const int32 BoneIndex) const
@@ -4142,6 +4162,8 @@
 	}
 }
 
+#if WITH_EDITOR
+
 void UAnimDatabaseFunction::Run_Implementation(UAnimDatabase* InDatabase, const FAnimDatabaseFrameRanges& InFrameRanges, const bool bShouldTransact) {}
 
 void UAnimDatabaseLibrary::GetAnimSequenceAnimNotifyClasses(TArray<TSubclassOf<UAnimNotify>>& OutClasses, UAnimSequence* AnimSequence)
@@ -4250,7 +4272,6 @@
 	const UE::Learning::FIndexSet UsedBones,
 	const bool bShouldTransact)
 {
-#if WITH_EDITOR
 	Controller->OpenBracket(LOCTEXT("KeyAnimSequence", "Keying Animation Sequence"), bShouldTransact);
 
 	const int32 FrameNum = PoseData.GetFrameNum();
@@ -4301,7 +4322,6 @@
 
 	Controller->NotifyPopulated();
 	Controller->CloseBracket(bShouldTransact);
-#endif
 }
 
 void UAnimDatabaseLibrary::AddAnimNotifiesToSequence(UAnimSequence* AnimSequence, const FName TrackName, const TSubclassOf<UAnimNotify>& Class, const TArray<float>& Times)
@@ -4311,7 +4331,6 @@
 
 void UAnimDatabaseLibrary::AddAnimNotifiesToSequenceArrayView(UAnimSequence* AnimSequence, const FName TrackName, const TSubclassOf<UAnimNotify>& Class, const TArrayView<const float> Times)
 {
-#if WITH_EDITOR
 	if (!AnimSequence)
 	{
 		UE_LOGFMT(LogAnimDatabase, Error, "AddAnimNotifiesToSequenceArrayView: AnimSequence is null.");
@@ -4359,12 +4378,10 @@
 	// Refresh all cached data
 	AnimSequence->RefreshCacheData();
 	AnimSequence->Modify(true);
-#endif
 }
 
 void UAnimDatabaseLibrary::RemoveAllAnimNotifiesAndAnimNotifyStatesFromSequence(UAnimSequence* AnimSequence)
 {
-#if WITH_EDITOR
 	if (!AnimSequence)
 	{
 		UE_LOGFMT(LogAnimDatabase, Error, "RemoveAllAnimNotifiesAndAnimNotifyStatesFromSequence: AnimSequence is null.");
@@ -4375,7 +4392,6 @@
 	AnimSequence->Notifies.Empty();
 	AnimSequence->RefreshCacheData();
 	AnimSequence->Modify(true);
-#endif
 }
 
 void UAnimDatabaseLibrary::AddAnimNotifyStatesToSequence(UAnimSequence* AnimSequence, const FName TrackName, const TSubclassOf<UAnimNotifyState>& Class, const TArray<float>& Times, const TArray<float>& Durations)
@@ -4385,7 +4401,6 @@
 
 void UAnimDatabaseLibrary::AddAnimNotifyStatesToSequenceArrayView(UAnimSequence* AnimSequence, const FName TrackName, const TSubclassOf<UAnimNotifyState>& Class, const TArrayView<const float> Times, const TArrayView<const float> Durations)
 {
-#if WITH_EDITOR
 	if (!AnimSequence)
 	{
 		UE_LOGFMT(LogAnimDatabase, Error, "AddAnimNotifyStatesToSequenceArrayView: AnimSequence is null.");
@@ -4435,7 +4450,6 @@
 	// Refresh all cached data
 	AnimSequence->RefreshCacheData();
 	AnimSequence->Modify(true);
-#endif
 }
 
 void UAnimDatabaseLibrary::EnableRootMotionAndForceRootLockOnSequences(const TArray<UAnimSequence*>& Sequences)
@@ -4445,7 +4459,6 @@
 
 void UAnimDatabaseLibrary::EnableRootMotionAndForceRootLockOnSequencesArrayView(TConstArrayView<UAnimSequence*> Sequences)
 {
-#if WITH_EDITOR
 	for (UAnimSequence* Sequence : Sequences)
 	{
 		if (Sequence)
@@ -4455,18 +4468,22 @@
 			Sequence->Modify(true);
 		}
 	}
-#endif
 }
 
 void UAnimDatabaseLibrary::DatabaseRemoveNotifies(UAnimDatabase* Database, const FAnimDatabaseFrameRanges& FrameRanges)
 {
-#if WITH_EDITOR
 	if (!Database)
 	{
 		UE_LOGFMT(LogAnimDatabase, Error, "DatabaseRemoveNotifies: Database is null.");
 		return;
 	}
 
+	if (!FrameRanges.IsValid())
+	{
+		UE_LOGFMT(LogAnimDatabase, Error, "DatabaseRemoveNotifies: Invalid Frame Ranges.");
+		return;
+	}
+
 	const int32 EntryNum = FrameRanges.FrameRangeSet->GetEntryNum();
 
 	for (int32 EntryIdx = 0; EntryIdx < EntryNum; EntryIdx++)
@@ -4533,18 +4550,22 @@
 			}
 		}
 	}
-#endif
 }
 
 void UAnimDatabaseLibrary::DatabaseExportNotifies(UAnimDatabase* Database, const FAnimDatabaseFrameRanges& FrameRanges, const FFilePath& Path)
 {
-#if WITH_EDITOR
 	if (!Database)
 	{
 		UE_LOGFMT(LogAnimDatabase, Error, "DatabaseExportNotifies: Database is null.");
 		return;
 	}
 
+	if (!FrameRanges.IsValid())
+	{
+		UE_LOGFMT(LogAnimDatabase, Error, "DatabaseExportNotifies: Invalid Frame Ranges.");
+		return;
+	}
+
 	TSharedRef<FJsonObject> Object = MakeShared<FJsonObject>();
 	Object->SetStringField(TEXT("Database"), Database->GetName());
 	Object->SetStringField(TEXT("Skeleton"), Database->GetSkeleton() ? Database->GetSkeleton()->GetName() : TEXT(""));
@@ -4679,7 +4700,6 @@
 			UE_LOGFMT(LogAnimDatabase, Error, "DatabaseExportNotifies: Failed to save to file {Filename}.", *Path.FilePath);
 		}
 	}
-#endif
 }
 
 void UAnimDatabaseLibrary::LoadAllBlueprintAssetsOfClass(UClass* Class)
@@ -4711,13 +4731,18 @@
 
 void UAnimDatabaseLibrary::DatabaseImportNotifies(UAnimDatabase* Database, const FAnimDatabaseFrameRanges& FrameRanges, const FFilePath& Path)
 {
-#if WITH_EDITOR
 	if (!Database)
 	{
 		UE_LOGFMT(LogAnimDatabase, Error, "DatabaseImportNotifies: Database is null.");
 		return;
 	}
 
+	if (!FrameRanges.IsValid())
+	{
+		UE_LOGFMT(LogAnimDatabase, Error, "DatabaseImportNotifies: Invalid Frame Ranges.");
+		return;
+	}
+
 	FString JsonString;
 	if (!FFileHelper::LoadFileToString(JsonString, *Path.FilePath))
 	{
@@ -4914,12 +4939,10 @@
 	{
 		UE_LOGFMT(LogAnimDatabase, Error, "DatabaseImportNotifies: Failed to parse JSON.");
 	}
-#endif
 }
 
 void UAnimDatabaseLibrary::DatabaseExportAsBVH(UAnimDatabase* Database, const FAnimDatabaseFrameRanges& FrameRanges, const FDirectoryPath& Path, const bool bAddRootAxesChangeRotation, const EAnimDatabaseRotationOrder RotationOrder)
 {
-#if WITH_EDITOR
 	if (!FrameRanges.IsValid())
 	{
 		UE_LOGFMT(LogAnimDatabase, Error, "DatabaseExportAsBVH: Invalid Frame Ranges.");
@@ -5072,12 +5095,10 @@
 			}
 		}
 	}
-#endif
 }
 
 void UAnimDatabaseLibrary::DatabaseExportAsAnimSequences(UAnimDatabase* Database, const FAnimDatabaseFrameRanges& FrameRanges, const FDirectoryPath& Path, const FString& AssetNameFormatString, const bool bExportMirrored, const bool bShouldTransact)
 {
-#if WITH_EDITOR
 	if (!FrameRanges.IsValid())
 	{
 		UE_LOGFMT(LogAnimDatabase, Error, "DatabaseExportAsAnimSequences: Invalid Frame Ranges.");
@@ -5228,12 +5249,10 @@
 			}
 		}
 	}
-#endif
 }
 
 void UAnimDatabaseLibrary::DatabaseExportRootMotion(UAnimDatabase* Database, const FAnimDatabaseFrameRanges& FrameRanges, const FFilePath& Path)
 {
-#if WITH_EDITOR
 	if (!FrameRanges.IsValid())
 	{
 		UE_LOGFMT(LogAnimDatabase, Error, "DatabaseExportRootMotion: Invalid Frame Ranges.");
@@ -5339,7 +5358,6 @@
 			UE_LOGFMT(LogAnimDatabase, Error, "DatabaseExportRootMotion: Failed to save to file {Filename}.", *Path.FilePath);
 		}
 	}
-#endif
 }
 
 void UAnimDatabaseLibrary::DatabaseAddCurve(UAnimDatabase* Database, const FName CurveName, const FAnimDatabaseFrameAttribute& FloatFrameAttribute, const bool bSparseKeys, const bool bShouldTransact)
@@ -5354,9 +5372,14 @@
 
 void UAnimDatabaseLibrary::DatabaseAddCurvesFromArrayViews(UAnimDatabase* Database, const TArrayView<const FName> CurveNames, const TArrayView<const FAnimDatabaseFrameAttribute> FloatFrameAttributes, const bool bSparseKeys, const bool bShouldTransact)
 {
-#if WITH_EDITOR
 	TRACE_CPUPROFILER_EVENT_SCOPE(UAnimDatabaseLibrary::DatabaseAddCurvesFromArrayViews);
 
+	if (!Database)
+	{
+		UE_LOGFMT(LogAnimDatabase, Error, "DatabaseAddCurvesFromArrayViews: Database is null.");
+		return;
+	}
+
 	if (CurveNames.Num() != FloatFrameAttributes.Num())
 	{
 		UE_LOGFMT(LogAnimDatabase, Error, "DatabaseAddCurvesFromArrayViews: Number of Curve names and Float FrameAttributes must match. Got {Actual} and {Expected}.", CurveNames.Num(), FloatFrameAttributes.Num());
@@ -5437,7 +5460,6 @@
 	}
 
 	UE::AnimDatabase::Private::CloseControllers(Controllers, bShouldTransact);
-#endif
 }
 
 void UAnimDatabaseLibrary::DatabaseAddAnimNotifyState(UAnimDatabase* Database, const FName AnimNotifyStateTrackName, const TSubclassOf<UAnimNotifyState> Class, const FAnimDatabaseFrameRanges& FrameRanges)
@@ -5452,9 +5474,14 @@
 
 void UAnimDatabaseLibrary::DatabaseAddAnimNotifyStatesFromArrayViews(UAnimDatabase* Database, const TArrayView<const FName> AnimNotifyStateTrackNames, const TArrayView<const TSubclassOf<UAnimNotifyState>> Classes, const TArrayView<const FAnimDatabaseFrameRanges> FrameRanges)
 {
-#if WITH_EDITOR
 	TRACE_CPUPROFILER_EVENT_SCOPE(UAnimDatabaseLibrary::DatabaseAddAnimNotifyStatesFromArrayViews);
 
+	if (!Database)
+	{
+		UE_LOGFMT(LogAnimDatabase, Error, "DatabaseAddAnimNotifyStatesFromArrayViews: Database is null.");
+		return;
+	}
+
 	if (AnimNotifyStateTrackNames.Num() != Classes.Num() ||
 		AnimNotifyStateTrackNames.Num() != FrameRanges.Num())
 	{
@@ -5524,12 +5551,10 @@
 			}
 		}
 	}
-#endif
 }
 
 void UAnimDatabaseLibrary::DatabaseAddAnimNotifyObject(UAnimDatabase* Database, const FAnimDatabaseFrames& Frames, const FName& TrackName, const UAnimNotify* AnimNotify)
 {
-#if WITH_EDITOR
 	TRACE_CPUPROFILER_EVENT_SCOPE(UAnimDatabaseLibrary::DatabaseAddAnimNotifyObject);
 
 	if (!Database)
@@ -5589,12 +5614,10 @@
 			AnimSequence->Modify(true);
 		}
 	}
-#endif
 }
 
 void UAnimDatabaseLibrary::DatabaseAddAnimNotifyStateObject(UAnimDatabase* Database, const FAnimDatabaseFrameRanges& FrameRanges, const FName& TrackName, const UAnimNotifyState* AnimNotifyState)
 {
-#if WITH_EDITOR
 	TRACE_CPUPROFILER_EVENT_SCOPE(UAnimDatabaseLibrary::DatabaseAddAnimNotifyStateObject);
 
 	if (!Database)
@@ -5655,7 +5678,6 @@
 			AnimSequence->Modify(true);
 		}
 	}
-#endif
 }
 
 void UAnimDatabaseLibrary::DatabaseAddAnimNotify(UAnimDatabase* Database, const FName AnimNotifyTrackName, const TSubclassOf<UAnimNotify> Class, const FAnimDatabaseFrames& Frames)
@@ -5670,7 +5692,6 @@
 
 void UAnimDatabaseLibrary::DatabaseAddAnimNotifiesFromArrayViews(UAnimDatabase* Database, const TArrayView<const FName> AnimNotifyTrackNames, const TArrayView<const TSubclassOf<UAnimNotify>> Classes, const TArrayView<const FAnimDatabaseFrames> Frames)
 {
-#if WITH_EDITOR
 	TRACE_CPUPROFILER_EVENT_SCOPE(UAnimDatabaseLibrary::DatabaseAddAnimNotifiesFromArrayViews);
 
 	if (AnimNotifyTrackNames.Num() != Classes.Num() ||
@@ -5745,7 +5766,6 @@
 			}
 		}
 	}
-#endif
 }
 
 void UAnimDatabaseLibrary::DatabaseSetLocalBoneTransform(UAnimDatabase* Database, const FName BoneName, const FAnimDatabaseFrameAttribute& TransformFrameAttribute, const bool bShouldTransact)
@@ -5760,7 +5780,6 @@
 
 void UAnimDatabaseLibrary::DatabaseSetLocalBoneTransformsFromArrayViews(UAnimDatabase* Database, const TArrayView<const FName> BoneNames, const TArrayView<const FAnimDatabaseFrameAttribute> TransformFrameAttributes, const bool bShouldTransact)
 {
-#if WITH_EDITOR
 	TRACE_CPUPROFILER_EVENT_SCOPE(UAnimDatabaseLibrary::DatabaseSetLocalBoneTransformsFromArrayViews);
 
 	if (BoneNames.Num() != TransformFrameAttributes.Num())
@@ -5874,12 +5893,10 @@
 	}
 
 	UE::AnimDatabase::Private::CloseControllers(Controllers, bShouldTransact);
-#endif
 }
 
 void UAnimDatabaseLibrary::DatabaseModifyPoseData(UAnimDatabase* Database, const FAnimDatabaseFrameRanges& FrameRanges, const TFunctionRef<void(const UE::AnimDatabase::FPoseDataView& InOutPoseData)> Function, const bool bShouldTransact)
 {
-#if WITH_EDITOR
 	const FAnimDatabaseFrameRanges FilteredFrameRanges = UAnimDatabaseFrameRangesLibrary::MakeFrameRangesFromNotMirrored(Database, FrameRanges);
 	
 	// Sample and modify pose data
@@ -5947,12 +5964,10 @@
 	}
 
 	UE::AnimDatabase::Private::CloseControllers(Controllers, bShouldTransact);
-#endif
 }
 
 void UAnimDatabaseLibrary::DatabaseMakeLooped(UAnimDatabase* Database, const FAnimDatabaseFrameRanges& FrameRanges, const float StartEndRatio, const float BlendInTime, const float BlendOutTime, const bool bShouldTransact)
 {
-#if WITH_EDITOR
 	TRACE_CPUPROFILER_EVENT_SCOPE(UAnimDatabaseLibrary::DatabaseMakeLooped);
 
 	DatabaseModifyPoseData(Database, FrameRanges, [Database, BlendInTime, BlendOutTime, StartEndRatio](const UE::AnimDatabase::FPoseDataView& InOutPoseData)
@@ -5960,12 +5975,10 @@
 			UE::AnimDatabase::PoseData::MakeLooped(InOutPoseData, Database->GetFrameRate(), StartEndRatio, BlendInTime, BlendOutTime);
 
 		}, bShouldTransact);
-#endif
 }
 
 void UAnimDatabaseLibrary::DatabasePatchRanges(UAnimDatabase* Database, const FAnimDatabaseFrameRanges& FrameRanges, const bool bApplyToRoot, const TArray<int32>& BoneIndices, const bool bShouldTransact)
 {
-#if WITH_EDITOR
 	TRACE_CPUPROFILER_EVENT_SCOPE(UAnimDatabaseLibrary::DatabasePatchRanges);
 
 	if (!Database) { return; }
@@ -5975,12 +5988,10 @@
 			UE::AnimDatabase::PoseData::PatchInterpolate(InOutPoseData, Database->GetFrameRate(), bApplyToRoot, BoneIndices);
 
 		}, bShouldTransact);
-#endif
 }
 
 void UAnimDatabaseLibrary::DatabaseSetToReferencePose(UAnimDatabase* Database, const FAnimDatabaseFrameRanges& FrameRanges, const bool bApplyToRoot, const TArray<int32>& IgnoreBoneIndices, const bool bShouldTransact)
 {
-#if WITH_EDITOR
 	TRACE_CPUPROFILER_EVENT_SCOPE(UAnimDatabaseLibrary::DatabasePatchRanges);
 
 	if (!Database || !Database->GetSkeleton()) { return; }
@@ -5996,7 +6007,6 @@
 			}
 
 		}, bShouldTransact);
-#endif
 }
 
 void UAnimDatabaseFunction_ExportBVH::Run_Implementation(UAnimDatabase* InDatabase, const FAnimDatabaseFrameRanges& InFrameRanges, const bool bShouldTransact)
@@ -6486,4 +6496,6 @@
 	}
 }
 
+#endif
+
 #undef LOCTEXT_NAMESPACE
diff --git a/Engine/Plugins/Experimental/Animation/AnimDatabase/Source/AnimDatabase/Public/AnimDatabase.h b/Engine/Plugins/Experimental/Animation/AnimDatabase/Source/AnimDatabase/Public/AnimDatabase.h
--- a/Engine/Plugins/Experimental/Animation/AnimDatabase/Source/AnimDatabase/Public/AnimDatabase.h
+++ b/Engine/Plugins/Experimental/Animation/AnimDatabase/Source/AnimDatabase/Public/AnimDatabase.h
@@ -90,6 +90,8 @@ class UAnimDatabaseViewportSettings : public UObject
 {
 	GENERATED_BODY()
 
+#if WITH_EDITORONLY_DATA
+
 public:
 
 	/** Mesh to use for previewing the character */
@@ -253,6 +255,8 @@ public:
 	/** Scale of the drawn scene origin */
 	UPROPERTY(EditAnywhere, Category = "Origin", meta = (ClampMin = "0", UIMin = "0", EditCondition = "bDrawOrigin", HideEditConditionToggle))
 	float DrawOriginScale = 25.0f;
+
+#endif
 };
 
 //--------------------------------------------------
@@ -293,10 +297,16 @@ class UAnimDatabaseQuery : public UObject
 
 public:
 
+#if WITH_EDITOR
+
 	/** Resets all of the transient properties/members. */
 	UFUNCTION(CallInEditor, Category = "Query", DisplayName = "Refresh")
 	UE_API void ResetTransient();
 
+#endif
+
+#if WITH_EDITORONLY_DATA
+
 	/** A pointer to the query database, used when constructing the query ranges */
 	UPROPERTY()
 	TObjectPtr<UAnimDatabase> Database;
@@ -330,6 +340,10 @@ public:
 	/** Seed used for the colors of the range identifiers */
 	int32 RangeIndentifierColorSeed = 1234;
 
+#endif
+
+#if WITH_EDITOR
+
 	/**
 	 * Checks if QueryRanges needs updating, and if so updates it. returns true if QueryRanges was updated (and therefore the UI may need updating),
 	 * otherwise false.
@@ -352,6 +366,10 @@ public:
 	 */
 	UE_API bool UpdateAdditionalTimelineObjects();
 
+#endif
+
+#if WITH_EDITORONLY_DATA
+
 	/** Additional Frames to display in the timeline */
 	UPROPERTY(EditAnywhere, Instanced, Category = "Timeline")
 	TArray<TObjectPtr<UAnimDatabaseFramesFunction>> AdditionalTimelineFrames;
@@ -413,9 +431,15 @@ public:
 	UPROPERTY(EditAnywhere, Instanced, Category = "Function")
 	TObjectPtr<UAnimDatabaseFunction> FunctionObject;
 
+#endif
+
+#if WITH_EDITOR
+
 	/** Evaluate the given function on the database with the chosen ranges */
 	UFUNCTION(CallInEditor, Category = "Function", DisplayName = "Run Function")
 	void RunFunction();
+
+#endif
 };
 
 //--------------------------------------------------
@@ -830,11 +854,15 @@ class UAnimDatabaseFunction : public UObject
 {
 	GENERATED_BODY()
 
+#if WITH_EDITOR
+
 public:
 
 	/** Run the database function on the provided database and frame ranges */
 	UFUNCTION(BlueprintNativeEvent, BlueprintCallable, Category = "AnimDatabase", meta = (ForceAsFunction))
 	UE_API void Run(UAnimDatabase* InDatabase, const FAnimDatabaseFrameRanges& InFrameRanges, const bool bShouldTransact);
+
+#endif
 };
 
 /**
@@ -849,6 +877,8 @@ class UAnimDatabaseLibrary : public UBlueprintFunctionLibrary
 {
 	GENERATED_BODY()
 
+#if WITH_EDITOR
+
 public:
 
 	/** Gets all of the AnimNotify classes present in a given AnimSequence */
@@ -984,6 +1014,7 @@ private:
 	/** Loads into memory all of the blueprint assets of a given class. */
 	static void LoadAllBlueprintAssetsOfClass(UClass* Class);
 
+#endif
 };
 
 
@@ -1007,7 +1038,13 @@ public:
 	UPROPERTY(EditAnywhere, Category = "Settings")
 	EAnimDatabaseRotationOrder RotationOrder = EAnimDatabaseRotationOrder::XYZ;
 
+#if WITH_EDITOR
+
+public:
+
 	UE_API virtual void Run_Implementation(UAnimDatabase* InDatabase, const FAnimDatabaseFrameRanges& InFrameRanges, const bool bShouldTransact) override;
+
+#endif
 };
 
 /** Function for exporting all of the AnimNotifies and AnimNotifyStates from the given animations  */
@@ -1026,7 +1063,13 @@ public:
 	UPROPERTY(EditAnywhere, Category = "Settings")
 	FString FileName = TEXT("Notifies.json");
 
+#if WITH_EDITOR
+
+public:
+
 	UE_API virtual void Run_Implementation(UAnimDatabase* InDatabase, const FAnimDatabaseFrameRanges& InFrameRanges, const bool bShouldTransact) override;
+
+#endif
 };
 
 /** Function for importing AnimNotifies and AnimNotifyStates onto the given animations  */
@@ -1041,7 +1084,13 @@ public:
 	UPROPERTY(EditAnywhere, Category = "Settings", meta = (RelativePath))
 	FFilePath ImportPath;
 
+#if WITH_EDITOR
+
+public:
+
 	UE_API virtual void Run_Implementation(UAnimDatabase* InDatabase, const FAnimDatabaseFrameRanges& InFrameRanges, const bool bShouldTransact) override;
+
+#endif
 };
 
 /** Function for removing AnimNotifies and AnimNotifyStates from the given animations  */
@@ -1050,9 +1099,13 @@ class UAnimDatabaseFunction_RemoveNotifies : public UAnimDatabaseFunction
 {
 	GENERATED_BODY()
 
+#if WITH_EDITOR
+
 public:
 
 	UE_API virtual void Run_Implementation(UAnimDatabase* InDatabase, const FAnimDatabaseFrameRanges& InFrameRanges, const bool bShouldTransact) override;
+
+#endif
 };
 
 /** Function for placing IK bones at the correct hand and foot locations  */
@@ -1061,6 +1114,8 @@ class UAnimDatabaseFunction_UpdateIKBones : public UAnimDatabaseFunction
 {
 	GENERATED_BODY()
 
+#if WITH_EDITORONLY_DATA
+
 public:
 
 	UPROPERTY(EditAnywhere, Category = "Bones")
@@ -1095,9 +1150,15 @@ public:
 	UPROPERTY(EditAnywhere, Category = "Bones")
 	FName IKHandBoneR = TEXT("ik_hand_r");
 
+#endif
+
+#if WITH_EDITOR
+
 public:
 
 	UE_API virtual void Run_Implementation(UAnimDatabase* InDatabase, const FAnimDatabaseFrameRanges& InFrameRanges, const bool bShouldTransact) override;
+
+#endif
 };
 
 /** Function for computing the root transform using the look-at direction  */
@@ -1106,6 +1167,8 @@ class UAnimDatabaseFunction_SetRootFromLookAt : public UAnimDatabaseFunction
 {
 	GENERATED_BODY()
 
+#if WITH_EDITORONLY_DATA
+
 public:
 
 	/** Root bone name */
@@ -1166,9 +1229,15 @@ public:
 	UPROPERTY(EditAnywhere, Category = "Settings", meta = (ForceUnits = "s", EditCondition = "!bUseSavGolFilter", HideEditConditionToggle))
 	float DirectionSmoothingTime = 0.25f;
 
+#endif
+
+#if WITH_EDITOR
+
 public:
 
 	UE_API virtual void Run_Implementation(UAnimDatabase* InDatabase, const FAnimDatabaseFrameRanges& InFrameRanges, const bool bShouldTransact) override;
+
+#endif
 };
 
 /** Function for baking root motion into the pelvis  */
@@ -1177,6 +1246,8 @@ class UAnimDatabaseFunction_BakeRootMotionIntoPelvis : public UAnimDatabaseFunct
 {
 	GENERATED_BODY()
 
+#if WITH_EDITORONLY_DATA
+
 public:
 
 	UPROPERTY(EditAnywhere, Category = "Bones")
@@ -1185,9 +1256,15 @@ public:
 	UPROPERTY(EditAnywhere, Category = "Bones")
 	FName PelvisBone = TEXT("pelvis");
 
+#endif
+
+#if WITH_EDITOR
+
 public:
 
 	UE_API virtual void Run_Implementation(UAnimDatabase* InDatabase, const FAnimDatabaseFrameRanges& InFrameRanges, const bool bShouldTransact) override;
+
+#endif
 };
 
 /** Function for making the root motion start at the origin for all the provided ranges  */
@@ -1196,14 +1273,22 @@ class UAnimDatabaseFunction_MakeRootStartAtOrigin : public UAnimDatabaseFunction
 {
 	GENERATED_BODY()
 
+#if WITH_EDITORONLY_DATA
+
 public:
 
 	UPROPERTY(EditAnywhere, Category = "Bones")
 	FName RootBone = TEXT("root");
 
+#endif
+
+#if WITH_EDITOR
+
 public:
 
 	UE_API virtual void Run_Implementation(UAnimDatabase* InDatabase, const FAnimDatabaseFrameRanges& InFrameRanges, const bool bShouldTransact) override;
+
+#endif
 };
 
 /** Function for exporting ranges to AnimSequences  */
@@ -1212,6 +1297,8 @@ class UAnimDatabaseFunction_ExportToAnimSequences : public UAnimDatabaseFunction
 {
 	GENERATED_BODY()
 
+#if WITH_EDITORONLY_DATA
+
 public:
 
 	UPROPERTY(EditAnywhere, Category = "Settings")
@@ -1223,9 +1310,15 @@ public:
 	UPROPERTY(EditAnywhere, Category = "Settings")
 	bool bExportMirrored = false;
 
+#endif
+
+#if WITH_EDITOR
+
 public:
 
 	UE_API virtual void Run_Implementation(UAnimDatabase* InDatabase, const FAnimDatabaseFrameRanges& InFrameRanges, const bool bShouldTransact) override;
+
+#endif
 };
 
 /** Function for exporting root motion from the given animations  */
@@ -1234,6 +1327,8 @@ class UAnimDatabaseFunction_ExportRootMotion : public UAnimDatabaseFunction
 {
 	GENERATED_BODY()
 
+#if WITH_EDITORONLY_DATA
+
 public:
 
 	/** Directory to export the root motion to. */
@@ -1244,7 +1339,15 @@ public:
 	UPROPERTY(EditAnywhere, Category = "Settings")
 	FString FileName = TEXT("RootMotion.json");
 
+#endif
+
+#if WITH_EDITOR
+
+public:
+
 	UE_API virtual void Run_Implementation(UAnimDatabase* InDatabase, const FAnimDatabaseFrameRanges& InFrameRanges, const bool bShouldTransact) override;
+
+#endif
 };
 
 /** Function for making ranges looping  */
@@ -1253,6 +1356,8 @@ class UAnimDatabaseFunction_MakeLooped : public UAnimDatabaseFunction
 {
 	GENERATED_BODY()
 
+#if WITH_EDITORONLY_DATA
+
 public:
 
 	UPROPERTY(EditAnywhere, Category = "Settings", meta = (UIMin = "0.0", UIMax = "1.0", ClampMin = "0.0", ClampMax = "1.0"))
@@ -1264,9 +1369,15 @@ public:
 	UPROPERTY(EditAnywhere, Category = "Settings", meta = (ForceUnits = "s"))
 	float BlendOutTime = 1.0f;
 
+#endif
+
+#if WITH_EDITOR
+
 public:
 
 	UE_API virtual void Run_Implementation(UAnimDatabase* InDatabase, const FAnimDatabaseFrameRanges& InFrameRanges, const bool bShouldTransact) override;
+
+#endif
 };
 
 /** Function for patching ranges  */
@@ -1275,6 +1386,8 @@ class UAnimDatabaseFunction_PatchBones : public UAnimDatabaseFunction
 {
 	GENERATED_BODY()
 
+#if WITH_EDITORONLY_DATA
+
 public:
 
 	UPROPERTY(EditAnywhere, Category = "Settings")
@@ -1286,9 +1399,15 @@ public:
 	UPROPERTY(EditAnywhere, Category = "Settings", meta = (EditCondition = "!bEntirePose", HideEditConditionToggle))
 	TArray<FName> BoneNames;
 
+#endif
+
+#if WITH_EDITOR
+
 public:
 
 	UE_API virtual void Run_Implementation(UAnimDatabase* InDatabase, const FAnimDatabaseFrameRanges& InFrameRanges, const bool bShouldTransact) override;
+
+#endif
 };
 
 /** Function for setting to the reference pose  */
@@ -1297,6 +1416,8 @@ class UAnimDatabaseFunction_SetToReferencePose : public UAnimDatabaseFunction
 {
 	GENERATED_BODY()
 
+#if WITH_EDITORONLY_DATA
+
 public:
 
 	UPROPERTY(EditAnywhere, Category = "Settings")
@@ -1305,9 +1426,15 @@ public:
 	UPROPERTY(EditAnywhere, Category = "Settings")
 	TArray<FName> IgnoreBoneNames;
 
+#endif
+
+#if WITH_EDITOR
+
 public:
 
 	UE_API virtual void Run_Implementation(UAnimDatabase* InDatabase, const FAnimDatabaseFrameRanges& InFrameRanges, const bool bShouldTransact) override;
+
+#endif
 };
 
 /** Function for applying a transform to a bone  */
@@ -1316,6 +1443,8 @@ class UAnimDatabaseFunction_BoneApplyTransform : public UAnimDatabaseFunction
 {
 	GENERATED_BODY()
 
+#if WITH_EDITORONLY_DATA
+
 public:
 
 	UPROPERTY(EditAnywhere, Category = "Settings")
@@ -1327,9 +1456,15 @@ public:
 	UPROPERTY(EditAnywhere, Category = "Settings")
 	FTransform Transform = FTransform::Identity;
 
+#endif
+
+#if WITH_EDITOR
+
 public:
 
 	UE_API virtual void Run_Implementation(UAnimDatabase* InDatabase, const FAnimDatabaseFrameRanges& InFrameRanges, const bool bShouldTransact) override;
+
+#endif
 };
 
 /** Function for adding AnimNotifyStates from a reference object  */
@@ -1338,6 +1473,8 @@ class UAnimDatabaseFunction_AddAnimNotifyState : public UAnimDatabaseFunction
 {
 	GENERATED_BODY()
 
+#if WITH_EDITORONLY_DATA
+
 public:
 
 	UPROPERTY(EditAnywhere, Category = "Settings")
@@ -1346,9 +1483,15 @@ public:
 	UPROPERTY(EditAnywhere, Instanced, Category = "Settings")
 	TObjectPtr<UAnimNotifyState> AnimNotifyState;
 
+#endif
+
+#if WITH_EDITOR
+
 public:
 
 	UE_API virtual void Run_Implementation(UAnimDatabase* InDatabase, const FAnimDatabaseFrameRanges& InFrameRanges, const bool bShouldTransact) override;
+
+#endif
 };
 
 /** Function that can execute a sequence of sub-functions  */
@@ -1357,14 +1500,22 @@ class UAnimDatabaseFunction_Sequence : public UAnimDatabaseFunction
 {
 	GENERATED_BODY()
 
+#if WITH_EDITORONLY_DATA
+
 public:
 
 	UPROPERTY(EditAnywhere, Instanced, Category = "Settings")
 	TArray<TObjectPtr<UAnimDatabaseFunction>> Functions;
 
+#endif
+
+#if WITH_EDITOR
+
 public:
 
 	UE_API virtual void Run_Implementation(UAnimDatabase* InDatabase, const FAnimDatabaseFrameRanges& InFrameRanges, const bool bShouldTransact) override;
+
+#endif
 };
 
 /** Function for adding an AnimNotify from a frame function  */
@@ -1373,6 +1524,8 @@ class UAnimDatabaseFunction_AddAnimNotifyAtFrames : public UAnimDatabaseFunction
 {
 	GENERATED_BODY()
 
+#if WITH_EDITORONLY_DATA
+
 public:
 
 	UPROPERTY(EditAnywhere, Category = "Settings")
@@ -1384,9 +1537,15 @@ public:
 	UPROPERTY(EditAnywhere, Instanced, Category = "Settings")
 	TObjectPtr<UAnimDatabaseFramesFunction> Frames;
 
+#endif
+
+#if WITH_EDITOR
+
 public:
 
 	UE_API virtual void Run_Implementation(UAnimDatabase* InDatabase, const FAnimDatabaseFrameRanges& InFrameRanges, const bool bShouldTransact) override;
+
+#endif
 };
 
 /** Function for adding an AnimNotifyState from a frame ranges function  */
@@ -1395,6 +1554,8 @@ class UAnimDatabaseFunction_AddAnimNotifyStateAtFrameRanges : public UAnimDataba
 {
 	GENERATED_BODY()
 
+#if WITH_EDITORONLY_DATA
+
 public:
 
 	UPROPERTY(EditAnywhere, Category = "Settings")
@@ -1406,9 +1567,15 @@ public:
 	UPROPERTY(EditAnywhere, Instanced, Category = "Settings")
 	TObjectPtr<UAnimDatabaseFrameRangesFunction> FrameRanges;
 
+#endif
+
+#if WITH_EDITOR
+
 public:
 
 	UE_API virtual void Run_Implementation(UAnimDatabase* InDatabase, const FAnimDatabaseFrameRanges& InFrameRanges, const bool bShouldTransact) override;
+
+#endif
 };
 
 /** Function for adding a curve from a frame attribute function  */
@@ -1417,6 +1584,8 @@ class UAnimDatabaseFunction_AddCurveFromFrameAttribute : public UAnimDatabaseFun
 {
 	GENERATED_BODY()
 
+#if WITH_EDITORONLY_DATA
+
 public:
 
 	UPROPERTY(EditAnywhere, Category = "Settings")
@@ -1428,9 +1597,15 @@ public:
 	UPROPERTY(EditAnywhere, Category = "Settings")
 	bool bSparseKeys = true;
 
+#endif
+
+#if WITH_EDITOR
+
 public:
 
 	UE_API virtual void Run_Implementation(UAnimDatabase* InDatabase, const FAnimDatabaseFrameRanges& InFrameRanges, const bool bShouldTransact) override;
+
+#endif
 };
 
 /** Function to set bone to root location at frame  */
@@ -1439,6 +1614,8 @@ class UAnimDatabaseFunction_SetBoneToRootTransformAtFrames : public UAnimDatabas
 {
 	GENERATED_BODY()
 
+#if WITH_EDITORONLY_DATA
+
 public:
 
 	UPROPERTY(EditAnywhere, Category = "Settings")
@@ -1447,9 +1624,15 @@ public:
 	UPROPERTY(EditAnywhere, Instanced, Category = "Settings")
 	TObjectPtr<UAnimDatabaseFramesFunction> Frames;
 
+#endif
+
+#if WITH_EDITOR
+
 public:
 
 	UE_API virtual void Run_Implementation(UAnimDatabase* InDatabase, const FAnimDatabaseFrameRanges& InFrameRanges, const bool bShouldTransact) override;
+
+#endif
 };
 
 /** Computes various statistics about the provided ranges */
@@ -1458,15 +1641,23 @@ class UAnimDatabaseFunction_Statistics : public UAnimDatabaseFunction
 {
 	GENERATED_BODY()
 
+#if WITH_EDITORONLY_DATA
+
 public:
 
 	/** Statistics for the given ranges */
 	UPROPERTY(VisibleAnywhere, Category = "Settings")
 	FAnimDatabaseFrameRangesStatistics Statistics;
 
+#endif
+
+#if WITH_EDITOR
+
 public:
 
 	UE_API virtual void Run_Implementation(UAnimDatabase* InDatabase, const FAnimDatabaseFrameRanges& InFrameRanges, const bool bShouldTransact) override;
+
+#endif
 };
 
 /** Computes various stats about root speeds for the provide frame ranges */
@@ -1475,6 +1666,8 @@ class UAnimDatabaseFunction_SpeedStatistics : public UAnimDatabaseFunction
 {
 	GENERATED_BODY()
 
+#if WITH_EDITORONLY_DATA
+
 public:
 
 	/** Frame ranges over which to compute the statistics */
@@ -1507,9 +1700,15 @@ public:
 	UPROPERTY(VisibleAnywhere, Category = "Settings", meta = (ForceUnits = "deg/s"))
 	float MaximumAngularVelocity = 0.0f;
 
+#endif
+
+#if WITH_EDITOR
+
 public:
 
 	UE_API virtual void Run_Implementation(UAnimDatabase* InDatabase, const FAnimDatabaseFrameRanges& InFrameRanges, const bool bShouldTransact) override;
+
+#endif
 };
 
 /** Computes various statistics regarding transitions given by the provided FrameRanges  */
@@ -1518,6 +1717,8 @@ class UAnimDatabaseFunction_TransitionStatistics : public UAnimDatabaseFunction
 {
 	GENERATED_BODY()
 
+#if WITH_EDITORONLY_DATA
+
 public:
 
 	/** Frame ranges over which to compute the statistics */
@@ -1562,9 +1763,15 @@ public:
 	UPROPERTY(VisibleAnywhere, Category = "Settings", meta = (ForceUnits = "deg"))
 	float MaximumAngle = 0.0f;
 
+#endif
+
+#if WITH_EDITOR
+
 public:
 
 	UE_API virtual void Run_Implementation(UAnimDatabase* InDatabase, const FAnimDatabaseFrameRanges& InFrameRanges, const bool bShouldTransact) override;
+
+#endif
 };
 
 //--------------------------------------------------
@@ -1577,6 +1784,8 @@ class UAnimDatabaseDebugDraw : public UObject
 {
 	GENERATED_BODY()
 
+#if WITH_EDITOR
+
 public:
 
 	/** Make any required frame attributes or other data. This function is called once when the class is initialized */
@@ -1598,6 +1807,8 @@ public:
 		const int32 RangeStart, 
 		const int32 RangeLength,
 		const FLinearColor& IdentifierColor);
+
+#endif
 };
 
 /** Debug Draw class for drawing contact curves */
@@ -1606,6 +1817,10 @@ class UAnimDatabaseDebugDraw_ContactCurves : public UAnimDatabaseDebugDraw
 {
 	GENERATED_BODY()
 
+#if WITH_EDITORONLY_DATA
+
+public:
+
 	/** Name of the left toe bone */
 	UPROPERTY(EditAnywhere, Category = "Settings")
 	FName LeftToeBoneName = TEXT("ball_l");
@@ -1626,7 +1841,7 @@ class UAnimDatabaseDebugDraw_ContactCurves : public UAnimDatabaseDebugDraw
 	UPROPERTY(EditAnywhere, Category = "Settings", meta = (ForceUnits = "cm"))
 	float ContactCircleRadius = 15.0f;
 
-public:
+protected:
 
 	/** Left contact curve attribute */
 	FAnimDatabaseFrameAttribute LeftContactCurve;
@@ -1634,6 +1849,12 @@ public:
 	/** Right contact curve attribute */
 	FAnimDatabaseFrameAttribute RightContactCurve;
 
+#endif
+
+#if WITH_EDITOR
+
+public:
+
 	virtual void InitializeDrawDebug_Implementation(const UAnimDatabase* InDatabase, const FAnimDatabaseFrameRanges& InFrameRanges) override;
 
 	virtual void DrawDebug_Implementation(
@@ -1647,6 +1868,8 @@ public:
 		const int32 RangeStart,
 		const int32 RangeLength,
 		const FLinearColor& IdentifierColor) override;
+
+#endif
 };
 
 
@@ -1656,6 +1879,10 @@ class UAnimDatabaseDebugDraw_RootVelocity : public UAnimDatabaseDebugDraw
 {
 	GENERATED_BODY()
 
+#if WITH_EDITORONLY_DATA
+
+public:
+
 	/** Graph axes maximum for linear velocity */
 	UPROPERTY(EditAnywhere, Category = "Settings", meta = (ForceForceUnits = "cm/s"))
 	float MaxLinearVelocity = 500.0f;
@@ -1664,7 +1891,7 @@ class UAnimDatabaseDebugDraw_RootVelocity : public UAnimDatabaseDebugDraw
 	UPROPERTY(EditAnywhere, Category = "Settings", meta = (ForceUnits = "deg/s"))
 	float MaxAngularVelocity = 180.0f;
 
-public:
+protected:
 
 	/** Root Linear Velocity attribute */
 	FAnimDatabaseFrameAttribute RootLinearVelocity;
@@ -1672,6 +1899,12 @@ public:
 	/** Root Angular Velocity attribute */
 	FAnimDatabaseFrameAttribute RootAngularVelocity;
 
+#endif
+
+#if WITH_EDITOR
+
+public:
+
 	virtual void InitializeDrawDebug_Implementation(const UAnimDatabase* InDatabase, const FAnimDatabaseFrameRanges& InFrameRanges) override;
 
 	virtual void DrawDebug_Implementation(
@@ -1685,6 +1918,8 @@ public:
 		const int32 RangeStart,
 		const int32 RangeLength,
 		const FLinearColor& IdentifierColor) override;
+
+#endif
 };
 
 
@@ -1694,6 +1929,10 @@ class UAnimDatabaseDebugDraw_ChairOnAttach : public UAnimDatabaseDebugDraw
 {
 	GENERATED_BODY()
 
+#if WITH_EDITORONLY_DATA
+
+public:
+
 	/** Name of the attach bone */
 	UPROPERTY(EditAnywhere, Category = "Settings")
 	FName AttachBoneName = TEXT("attach");
@@ -1702,11 +1941,17 @@ class UAnimDatabaseDebugDraw_ChairOnAttach : public UAnimDatabaseDebugDraw
 	UPROPERTY(EditAnywhere, Instanced, Category = "Settings")
 	TObjectPtr<UAnimDatabaseFrameRangesFunction> AttachRanges;
 
-public:
+protected:
 
 	/** Valid ranges for attachment */
 	FAnimDatabaseFrameRanges ValidAttachRanges;
 
+#endif
+
+#if WITH_EDITOR
+
+public:
+
 	virtual void InitializeDrawDebug_Implementation(const UAnimDatabase* InDatabase, const FAnimDatabaseFrameRanges& InFrameRanges) override;
 
 	virtual void DrawDebug_Implementation(
@@ -1720,6 +1965,8 @@ public:
 		const int32 RangeStart,
 		const int32 RangeLength,
 		const FLinearColor& IdentifierColor) override;
+
+#endif
 };
 
 #undef UE_API
\ No newline at end of file
