AnimDatabase: Added function for adjusting neck rotation

#Synced-CL 51119969

diff --git a/Engine/Plugins/Experimental/Animation/AnimDatabase/Source/AnimDatabase/Private/AnimDatabase.cpp b/Engine/Plugins/Experimental/Animation/AnimDatabase/Source/AnimDatabase/Private/AnimDatabase.cpp
--- a/Engine/Plugins/Experimental/Animation/AnimDatabase/Source/AnimDatabase/Private/AnimDatabase.cpp
+++ b/Engine/Plugins/Experimental/Animation/AnimDatabase/Source/AnimDatabase/Private/AnimDatabase.cpp
@@ -6267,6 +6267,29 @@ void UAnimDatabaseFunction_SetBoneToRootTransformAtFrames::Run_Implementation(UA
 
 }
 
+void UAnimDatabaseFunction_AdjustNeckRotation::Run_Implementation(UAnimDatabase* InDatabase, const FAnimDatabaseFrameRanges& InFrameRanges, const bool bShouldTransact)
+{
+	const FAnimDatabaseFrameAttribute Neck0Transform = UAnimDatabaseFrameAttributeLibrary::MakeBoneLocalTransformFrameAttributeFromName(InDatabase, InFrameRanges, Neck0Bone);
+	const FAnimDatabaseFrameAttribute Neck1Transform = UAnimDatabaseFrameAttributeLibrary::MakeBoneLocalTransformFrameAttributeFromName(InDatabase, InFrameRanges, Neck1Bone);
+	const FAnimDatabaseFrameAttribute HeadTransform = UAnimDatabaseFrameAttributeLibrary::MakeBoneLocalTransformFrameAttributeFromName(InDatabase, InFrameRanges, HeadBone);
+
+	UAnimDatabaseLibrary::DatabaseSetLocalBoneTransforms(
+		InDatabase,
+		{ Neck0Bone, Neck1Bone, HeadBone },
+		{
+			UAnimDatabaseFrameAttributeLibrary::FrameAttributeMultiply(UAnimDatabaseFrameAttributeLibrary::MakeTransformFrameAttributeNoScale(
+				UAnimDatabaseFrameAttributeLibrary::MakeLocationFrameAttributeFromConstant(InFrameRanges, FVector::ZeroVector),
+				UAnimDatabaseFrameAttributeLibrary::MakeRotationFrameAttributeFromConstant(InFrameRanges, FQuat::MakeFromRotationVector(FMath::DegreesToRadians(Neck0Rotation) * LocalAxis).Rotator())), Neck0Transform),
+			UAnimDatabaseFrameAttributeLibrary::FrameAttributeMultiply(UAnimDatabaseFrameAttributeLibrary::MakeTransformFrameAttributeNoScale(
+				UAnimDatabaseFrameAttributeLibrary::MakeLocationFrameAttributeFromConstant(InFrameRanges, FVector::ZeroVector),
+				UAnimDatabaseFrameAttributeLibrary::MakeRotationFrameAttributeFromConstant(InFrameRanges, FQuat::MakeFromRotationVector(FMath::DegreesToRadians(Neck1Rotation) * LocalAxis).Rotator())), Neck1Transform),
+			UAnimDatabaseFrameAttributeLibrary::FrameAttributeMultiply(UAnimDatabaseFrameAttributeLibrary::MakeTransformFrameAttributeNoScale(
+				UAnimDatabaseFrameAttributeLibrary::MakeLocationFrameAttributeFromConstant(InFrameRanges, FVector::ZeroVector),
+				UAnimDatabaseFrameAttributeLibrary::MakeRotationFrameAttributeFromConstant(InFrameRanges, FQuat::MakeFromRotationVector(FMath::DegreesToRadians(HeadRotation) * LocalAxis).Rotator())), HeadTransform)
+		},
+		bShouldTransact);
+}
+
 void UAnimDatabaseFunction_Statistics::Run_Implementation(UAnimDatabase* InDatabase, const FAnimDatabaseFrameRanges& InFrameRanges, const bool bShouldTransact)
 {
 	UAnimDatabaseFrameRangesLibrary::FrameRangesStatistics(Statistics, InDatabase, InFrameRanges);
diff --git a/Engine/Plugins/Experimental/Animation/AnimDatabase/Source/AnimDatabase/Public/AnimDatabase.h b/Engine/Plugins/Experimental/Animation/AnimDatabase/Source/AnimDatabase/Public/AnimDatabase.h
--- a/Engine/Plugins/Experimental/Animation/AnimDatabase/Source/AnimDatabase/Public/AnimDatabase.h
+++ b/Engine/Plugins/Experimental/Animation/AnimDatabase/Source/AnimDatabase/Public/AnimDatabase.h
@@ -1635,6 +1635,55 @@ public:
 #endif
 };
 
+/** Function to adjust neck rotation to remove potential artifacts after retargeting */
+UCLASS(BlueprintType, Blueprintable, meta = (DisplayName = "Adjust Neck Rotation"))
+class UAnimDatabaseFunction_AdjustNeckRotation : public UAnimDatabaseFunction
+{
+	GENERATED_BODY()
+
+#if WITH_EDITORONLY_DATA
+
+public:
+
+	/** Name of the neck0 bone */
+	UPROPERTY(EditAnywhere, Category = "Settings")
+	FName Neck0Bone = TEXT("neck_01");
+
+	/** Name of the neck1 bone */
+	UPROPERTY(EditAnywhere, Category = "Settings")
+	FName Neck1Bone = TEXT("neck_02");
+
+	/** Name of the head bone */
+	UPROPERTY(EditAnywhere, Category = "Settings")
+	FName HeadBone = TEXT("head");
+
+	/** Rotation to apply to the neck0 bone */
+	UPROPERTY(EditAnywhere, Category = "Settings", meta = (ForceUnits = "deg"))
+	float Neck0Rotation = -18.0f;
+
+	/** Rotation to apply to the neck1 bone */
+	UPROPERTY(EditAnywhere, Category = "Settings", meta = (ForceUnits = "deg"))
+	float Neck1Rotation = 0.0f;
+
+	/** Rotation to apply to the head bone */
+	UPROPERTY(EditAnywhere, Category = "Settings", meta = (ForceUnits = "deg"))
+	float HeadRotation = 16.0f;
+
+	/** Local axis to apply rotation on */
+	UPROPERTY(EditAnywhere, Category = "Settings")
+	FVector LocalAxis = FVector(0.0, 0.0, 1.0);
+
+#endif
+
+#if WITH_EDITOR
+
+public:
+
+	UE_API virtual void Run_Implementation(UAnimDatabase* InDatabase, const FAnimDatabaseFrameRanges& InFrameRanges, const bool bShouldTransact) override;
+
+#endif
+};
+
 /** Computes various statistics about the provided ranges */
 UCLASS(BlueprintType, Blueprintable, meta = (DisplayName = "Statistics"))
 class UAnimDatabaseFunction_Statistics : public UAnimDatabaseFunction
