AnimGen: Move to Experimental

#Synced-CL 51372803

#Synced-CL 51372803

diff --git a/Engine/Plugins/Experimental/Animation/AnimDatabase/Source/AnimDatabase/Private/AnimDatabase.cpp b/Engine/Plugins/Experimental/Animation/AnimDatabase/Source/AnimDatabase/Private/AnimDatabase.cpp
--- a/Engine/Plugins/Experimental/Animation/AnimDatabase/Source/AnimDatabase/Private/AnimDatabase.cpp
+++ b/Engine/Plugins/Experimental/Animation/AnimDatabase/Source/AnimDatabase/Private/AnimDatabase.cpp
@@ -70,6 +70,28 @@ DEFINE_LOG_CATEGORY(LogAnimDatabase);
 
 namespace UE::AnimDatabase::Private
 {
+	void LocationOrderFromStrings(TArray<int32>& LocationOrder, const FString& A, const FString& B, const FString& C)
+	{
+		if      (A == TEXT("Xposition") && B == TEXT("Yposition") && C == TEXT("Zposition")) { LocationOrder = { 0, 1, 2 }; return; }
+		else if (A == TEXT("Xposition") && B == TEXT("Zposition") && C == TEXT("Yposition")) { LocationOrder = { 0, 2, 1 }; return; }
+		else if (A == TEXT("Yposition") && B == TEXT("Xposition") && C == TEXT("Zposition")) { LocationOrder = { 1, 0, 2 }; return; }
+		else if (A == TEXT("Yposition") && B == TEXT("Zposition") && C == TEXT("Xposition")) { LocationOrder = { 1, 2, 0 }; return; }
+		else if (A == TEXT("Zposition") && B == TEXT("Xposition") && C == TEXT("Yposition")) { LocationOrder = { 2, 0, 1 }; return; }
+		else if (A == TEXT("Zposition") && B == TEXT("Yposition") && C == TEXT("Xposition")) { LocationOrder = { 2, 1, 0 }; return; }
+		else { LocationOrder = { 0, 1, 2 }; return; }
+	}
+
+	EAnimDatabaseRotationOrder RotationOrderFromStrings(const FString& A, const FString& B, const FString& C)
+	{
+		if      (A == TEXT("Xrotation") && B == TEXT("Yrotation") && C == TEXT("Zrotation")) { return EAnimDatabaseRotationOrder::XYZ; }
+		else if (A == TEXT("Xrotation") && B == TEXT("Zrotation") && C == TEXT("Yrotation")) { return EAnimDatabaseRotationOrder::XZY; }
+		else if (A == TEXT("Yrotation") && B == TEXT("Xrotation") && C == TEXT("Zrotation")) { return EAnimDatabaseRotationOrder::YXZ; }
+		else if (A == TEXT("Yrotation") && B == TEXT("Zrotation") && C == TEXT("Xrotation")) { return EAnimDatabaseRotationOrder::YZX; }
+		else if (A == TEXT("Zrotation") && B == TEXT("Xrotation") && C == TEXT("Yrotation")) { return EAnimDatabaseRotationOrder::ZXY; }
+		else if (A == TEXT("Zrotation") && B == TEXT("Yrotation") && C == TEXT("Xrotation")) { return EAnimDatabaseRotationOrder::ZYX; }
+		else { return EAnimDatabaseRotationOrder::XYZ; }
+	}
+
 	FVector3f RotationOrderAxis(const EAnimDatabaseRotationOrder RotationOrder, const uint8 OrderIndex)
 	{
 		check(OrderIndex >= 0 && OrderIndex < 3);
@@ -2023,7 +2045,7 @@ namespace UE::AnimDatabase::Private
 
 		OutBoneOrder.Add(BoneIndex);
 
-		InOut += FString::Printf(TEXT("%sJOINT %s\n"), *InOutIndentation, *BoneNames[BoneIndex].ToString());
+		InOut += FString::Printf(TEXT("%sJOINT %s\n"), *InOutIndentation, *BoneNames[BoneIndex].ToString().Replace(TEXT(" "), TEXT("_")));
 		InOut += FString::Printf(TEXT("%s{\n"), *InOutIndentation);
 		InOutIndentation += TEXT("\t");
 
@@ -2152,7 +2174,7 @@ namespace UE::AnimDatabase::Private
 		InOutIndentation = TEXT("");
 
 		InOut += TEXT("HIERARCHY\n");
-		InOut += FString::Printf(TEXT("ROOT %s\n"), *BoneNames[0].ToString());
+		InOut += FString::Printf(TEXT("ROOT %s\n"), *BoneNames[0].ToString().Replace(TEXT(" "), TEXT("_")));
 		InOut += FString::Printf(TEXT("%s{\n"), *InOutIndentation);
 		InOutIndentation += TEXT("\t");
 
@@ -4995,22 +5017,31 @@ void UAnimDatabaseLibrary::DatabaseExportAsBVH(UAnimDatabase* Database, const FA
 		RefBoneRotations,
 		RotationOrder);
 
-	UE::AnimDatabase::FPoseData PoseData;
-	UE::AnimDatabase::FPoseGlobalBoneData GlobalPoseData;
-
-	const int32 EntryNum = FrameRanges.FrameRangeSet->GetEntryNum();
-
-	for (int32 EntryIdx = 0; EntryIdx < EntryNum; EntryIdx++)
-	{
-		const int32 SequenceIdx = FrameRanges.FrameRangeSet->GetEntrySequence(EntryIdx);
-		const int32 RangeNum = FrameRanges.FrameRangeSet->GetEntryRangeNum(EntryIdx);
-
-		if (UAnimSequence* AnimSequence = Database->GetAnimSequence(SequenceIdx))
-		{
-			for (int32 RangeIdx = 0; RangeIdx < RangeNum; RangeIdx++)
+	UE::Learning::FrameRangeSet::ParallelForEachRange(*FrameRanges.FrameRangeSet, [
+		&FrameRanges,
+		Database,
+		&BoneNames,
+		&BoneParents,
+		&BoneOrder,
+		bAddRootAxesChangeRotation,
+		RotationOrder,
+		&Path,
+		&BVHHeader](
+			const int32 TotalRangeIdx, 
+			const int32 EntryIdx, 
+			const int32 RangeIdx) {
+
+			const int32 SequenceIdx = FrameRanges.FrameRangeSet->GetEntrySequence(EntryIdx);
+			const int32 RangeNum = FrameRanges.FrameRangeSet->GetEntryRangeNum(EntryIdx);
+
+			if (UAnimSequence* AnimSequence = Database->GetAnimSequence(SequenceIdx))
 			{
 				const int32 FrameStart = FrameRanges.FrameRangeSet->GetEntryRangeStart(EntryIdx, RangeIdx);
 				const int32 FrameNum = FrameRanges.FrameRangeSet->GetEntryRangeLength(EntryIdx, RangeIdx);
+				const int32 BoneNum = Database->GetBoneNum();
+
+				UE::AnimDatabase::FPoseData PoseData;
+				UE::AnimDatabase::FPoseGlobalBoneData GlobalPoseData;
 
 				PoseData.Resize(FrameNum, BoneNum, {});
 				Database->GetPoseData(PoseData.View(), SequenceIdx, FrameStart, {});
@@ -5093,6 +5124,432 @@ void UAnimDatabaseLibrary::DatabaseExportAsBVH(UAnimDatabase* Database, const FA
 					UE_LOGFMT(LogAnimDatabase, Error, "DatabaseExportAsBVH: Failed to save to file {Filename}.", *FilePath);
 				}
 			}
+		});
+}
+
+void UAnimDatabaseLibrary::DatabaseImportFromBVH(UAnimDatabase* Database, const FDirectoryPath& ImportPath, const FDirectoryPath& AssetPath, const bool bRemoveRootAxesChangeRotation, const bool bAddToDatabase, const bool bIgnoreMirrored)
+{
+	if (!Database)
+	{
+		UE_LOGFMT(LogAnimDatabase, Error, "DatabaseImportFromBVH: Database is null.");
+		return;
+	}
+
+	if (ImportPath.Path.IsEmpty())
+	{
+		UE_LOGFMT(LogAnimDatabase, Error, "DatabaseImportFromBVH: ImportPath is empty.");
+		return;
+	}
+
+	if (AssetPath.Path.IsEmpty())
+	{
+		UE_LOGFMT(LogAnimDatabase, Error, "DatabaseImportFromBVH: AssetPath is empty.");
+		return;
+	}
+
+	if (!FPaths::DirectoryExists(ImportPath.Path))
+	{
+		UE_LOGFMT(LogAnimDatabase, Error, "DatabaseImportFromBVH: Directory {Directory} doesn't exist.", ImportPath.Path);
+		return;
+	}
+
+	if (!FPaths::DirectoryExists(AssetPath.Path))
+	{
+		UE_LOGFMT(LogAnimDatabase, Error, "DatabaseImportFromBVH: Directory {Directory} doesn't exist.", AssetPath.Path);
+		return;
+	}
+
+
+	IFileManager& FileManager = IFileManager::Get();
+
+	TArray<FString> FoundFiles;
+	FileManager.FindFiles(FoundFiles, *(ImportPath.Path + TEXT("/*.bvh")), true, false);
+
+	FScopedSlowTask SlowTask(FoundFiles.Num(), LOCTEXT("ImportingFromBVH", "Importing From BVH"));
+	SlowTask.MakeDialog(true);
+
+	for (const FString& File : FoundFiles)
+	{
+		SlowTask.EnterProgressFrame();
+		if (SlowTask.ShouldCancel()) { break; }
+
+		if (bIgnoreMirrored && File.EndsWith(TEXT("_mirrored.bvh"))) { continue; }
+
+		FString BVHContents;
+		if (FFileHelper::LoadFileToString(BVHContents, *(ImportPath.Path + TEXT("/") + File)))
+		{
+			int32 ActiveIdx = INDEX_NONE;
+			bool bEndSite = false;
+			int32 FrameNum = 0;
+			int32 FrameIdx = 0;
+			float FrameTime = 0.0f;
+			int32 ChannelNum = 0;
+			EAnimDatabaseRotationOrder RotationOrder = EAnimDatabaseRotationOrder::XYZ;
+			TArray<int32> LocationOrder = { 0, 1, 2 };
+
+			TArray<FString> BoneNames;
+			TArray<int32> BoneParents;
+			TArray<FQuat4f> BoneOrients;
+			TArray<FVector3f> BoneOffsets;
+			TLearningArray<2, FVector3f> BoneLocations;
+			TLearningArray<2, FQuat4f> BoneRotations;
+
+			TArray<FString> Lines;
+			BVHContents.ParseIntoArrayLines(Lines);
+			for (const FString& Line : Lines)
+			{
+				if (Line.Contains(TEXT("HIERARCHY"))) { continue; }
+				if (Line.Contains(TEXT("MOTION"))) { continue; }
+
+				if (Line.Contains(TEXT("ROOT")) || Line.Contains(TEXT("JOINT")))
+				{
+					TArray<FString> Parts;
+					Line.ParseIntoArray(Parts, TEXT(" "), true);
+					BoneNames.Add(ensure(Parts.Num() == 2) ? Parts[1] : TEXT(""));
+					BoneOrients.Add(FQuat4f::Identity);
+					BoneOffsets.Add(FVector3f::ZeroVector);
+					BoneParents.Add(ActiveIdx);
+					ActiveIdx = BoneParents.Num() - 1;
+					continue;
+				}
+
+				if (Line.Contains(TEXT("{"))) { continue; }
+
+				if (Line.Contains(TEXT("}")))
+				{
+					if (bEndSite)
+					{
+						bEndSite = false;
+					}
+					else
+					{
+						ActiveIdx = ActiveIdx == INDEX_NONE ? INDEX_NONE : BoneParents[ActiveIdx];
+					}
+					continue;
+				}
+
+				if (Line.Contains(TEXT("OFFSET")))
+				{
+					TArray<FString> Parts;
+					Line.ParseIntoArray(Parts, TEXT(" "), true);
+
+					if (!bEndSite && ActiveIdx != INDEX_NONE)
+					{
+						BoneOffsets[ActiveIdx] = ensure(Parts.Num() == 4) ? 
+							FVector3f(
+								FCString::Atof(*Parts[1]),
+								FCString::Atof(*Parts[2]),
+								FCString::Atof(*Parts[3])) : FVector3f::ZeroVector;
+					}
+					continue;
+				}
+
+				if (Line.Contains(TEXT("CHANNELS")))
+				{
+					TArray<FString> Parts;
+					Line.ParseIntoArray(Parts, TEXT(" "), true);
+					
+					if (ensure(Parts.Num() > 1))
+					{
+						ChannelNum = FCString::Atoi(*Parts[1]);
+					}
+
+					if (Parts.Num() == 2 + 3)
+					{
+						RotationOrder = UE::AnimDatabase::Private::RotationOrderFromStrings(Parts[2], Parts[3], Parts[4]);
+					}
+					else if (Parts.Num() == 2 + 6)
+					{
+						UE::AnimDatabase::Private::LocationOrderFromStrings(LocationOrder, Parts[2], Parts[3], Parts[4]);
+						RotationOrder = UE::AnimDatabase::Private::RotationOrderFromStrings(Parts[5], Parts[6], Parts[7]);
+					}
+					continue;
+				}
+
+				if (Line.Contains(TEXT("End Site")))
+				{
+					bEndSite = true;
+					continue;
+				}
+
+				if (Line.Contains(TEXT("Frames: ")))
+				{
+					TArray<FString> Parts;
+					Line.ParseIntoArray(Parts, TEXT(" "), true);
+					FrameNum = ensure(Parts.Num() == 2) ? FCString::Atoi(*Parts[1]) : 0;
+					BoneLocations.SetNumUninitialized({ FrameNum, BoneNames.Num() });
+					BoneRotations.SetNumUninitialized({ FrameNum, BoneNames.Num() });
+					for (int32 FillFrameIdx = 0; FillFrameIdx < FrameNum; FillFrameIdx++)
+					{
+						UE::Learning::Array::Copy(BoneLocations[FillFrameIdx], TLearningArrayView<1, const FVector3f>(BoneOffsets));
+						UE::Learning::Array::Copy(BoneRotations[FillFrameIdx], TLearningArrayView<1, const FQuat4f>(BoneOrients));
+					}
+					continue;
+				}
+
+				if (Line.Contains(TEXT("Frame Time: ")))
+				{
+					TArray<FString> Parts;
+					Line.ParseIntoArray(Parts, TEXT(" "), true);
+					FrameTime = ensure(Parts.Num() == 3) ? FCString::Atof(*Parts[2]) : 0;
+					continue;
+				}
+
+				const int32 BoneNum = BoneNames.Num();
+
+				TArray<FString> Parts;
+				Line.ParseIntoArray(Parts, TEXT(" "), true);
+
+				if (ChannelNum == 3)
+				{
+					if (ensure(Parts.Num() == ChannelNum * BoneNum + ChannelNum && FrameIdx < FrameNum))
+					{
+						BoneLocations[FrameIdx][0] = FVector3f(
+							FCString::Atof(*Parts[0]),
+							FCString::Atof(*Parts[1]), 
+							FCString::Atof(*Parts[2]));
+
+						for (int32 BoneIdx = 0; BoneIdx < BoneNum; BoneIdx++)
+						{
+							BoneRotations[FrameIdx][BoneIdx] = UE::AnimDatabase::Private::QuatFromEuler(FVector3f(
+								FMath::DegreesToRadians(FCString::Atof(*Parts[3 + BoneIdx * 3 + 0])),
+								FMath::DegreesToRadians(FCString::Atof(*Parts[3 + BoneIdx * 3 + 1])),
+								FMath::DegreesToRadians(FCString::Atof(*Parts[3 + BoneIdx * 3 + 2]))), RotationOrder);
+						}
+					}
+
+					FrameIdx++;
+					continue;
+				}
+				
+				if (ChannelNum == 6)
+				{
+					if (ensure(Parts.Num() == ChannelNum * BoneNum && FrameIdx < FrameNum))
+					{
+						for (int32 BoneIdx = 0; BoneIdx < BoneNum; BoneIdx++)
+						{
+							BoneLocations[FrameIdx][BoneIdx] = FVector3f(
+								FCString::Atof(*Parts[BoneIdx * 6 + LocationOrder[0]]),
+								FCString::Atof(*Parts[BoneIdx * 6 + LocationOrder[1]]),
+								FCString::Atof(*Parts[BoneIdx * 6 + LocationOrder[2]]));
+
+							BoneRotations[FrameIdx][BoneIdx] = UE::AnimDatabase::Private::QuatFromEuler(FVector3f(
+								FMath::DegreesToRadians(FCString::Atof(*Parts[BoneIdx * 6 + 3])),
+								FMath::DegreesToRadians(FCString::Atof(*Parts[BoneIdx * 6 + 4])),
+								FMath::DegreesToRadians(FCString::Atof(*Parts[BoneIdx * 6 + 5]))), RotationOrder);
+						}
+					}
+
+					FrameIdx++;
+					continue;
+				}
+
+				// Didn't parse line correctly...
+				ensure(false);
+			}
+
+			if (!ensure(FrameIdx == FrameNum && FrameNum > 0)) { continue; }
+
+			const int32 BVHBoneNum = BoneNames.Num();
+
+			if (bRemoveRootAxesChangeRotation && BVHBoneNum > 0)
+			{
+				for (int32 LocalFrameIdx = 0; LocalFrameIdx < FrameNum; LocalFrameIdx++)
+				{
+					const FQuat4f Adjustment = FQuat4f::MakeFromRotationVector(FVector3f(+UE_HALF_PI, 0.0f, 0.0f));
+					BoneRotations[LocalFrameIdx][0] = Adjustment * BoneRotations[LocalFrameIdx][0];
+					BoneLocations[LocalFrameIdx][0] = Adjustment.RotateVector(BoneLocations[LocalFrameIdx][0]);
+				}
+			}
+
+			const int32 DatabaseBoneNum = Database->GetBoneNum();
+
+			UE::AnimDatabase::FPoseData PoseData;
+			PoseData.Resize(FrameNum, DatabaseBoneNum);
+			UE::AnimDatabase::PoseData::Reset(PoseData.View(), Database->GetSkeleton()->GetReferenceSkeleton().GetRefBonePose());
+
+			TLearningArray<2, FVector3f> GlobalBoneLocations = BoneLocations;
+			TLearningArray<2, FQuat4f> GlobalBoneRotations = BoneRotations;
+			for (int32 LocalFrameIdx = 0; LocalFrameIdx < FrameNum; LocalFrameIdx++)
+			{
+				for (int32 BoneIdx = 0; BoneIdx < BVHBoneNum; BoneIdx++)
+				{
+					const int32 ParentIdx = BoneParents[BoneIdx];
+					check(ParentIdx < BoneIdx);
+					if (ParentIdx != INDEX_NONE)
+					{
+						GlobalBoneLocations[LocalFrameIdx][BoneIdx] = GlobalBoneRotations[LocalFrameIdx][ParentIdx].RotateVector(GlobalBoneLocations[LocalFrameIdx][BoneIdx]) + GlobalBoneLocations[LocalFrameIdx][ParentIdx];
+						GlobalBoneRotations[LocalFrameIdx][BoneIdx] = GlobalBoneRotations[LocalFrameIdx][ParentIdx] * GlobalBoneRotations[LocalFrameIdx][BoneIdx];
+					}
+				}
+			}
+
+			for (int32 LocalFrameIdx = 0; LocalFrameIdx < FrameNum; LocalFrameIdx++)
+			{
+				for (int32 BoneIdx = 0; BoneIdx < BVHBoneNum; BoneIdx++)
+				{
+					const int32 DatabaseBoneIdx = Database->FindBoneIndex(FName(BoneNames[BoneIdx]));
+					if (DatabaseBoneIdx == INDEX_NONE) { continue; }
+
+					const FQuat4f BoneRotation = GlobalBoneRotations[LocalFrameIdx][BoneIdx];
+					const FVector3f BoneLocation = GlobalBoneLocations[LocalFrameIdx][BoneIdx];
+
+					FMatrix44f BoneTransform = FTransform3f(GlobalBoneRotations[LocalFrameIdx][BoneIdx], GlobalBoneLocations[LocalFrameIdx][BoneIdx]).ToMatrixNoScale();
+					for (int32 Idx = 0; Idx < 4; Idx++) { Swap(BoneTransform.M[1][Idx], BoneTransform.M[2][Idx]); }
+
+					for (int32 Idx = 0; Idx < 3; Idx++)
+					{
+						const float Y = BoneTransform.M[1][Idx];
+						const float Z = BoneTransform.M[2][Idx];
+						BoneTransform.M[2][Idx] = +Y;
+						BoneTransform.M[1][Idx] = -Z;
+					}
+
+					FMatrix44f ParentTransform = FMatrix44f::Identity;
+					if (BoneParents[BoneIdx] != INDEX_NONE)
+					{
+						FVector3f ParentLocation = GlobalBoneLocations[LocalFrameIdx][BoneParents[BoneIdx]];
+						FQuat4f ParentRotation = GlobalBoneRotations[LocalFrameIdx][BoneParents[BoneIdx]];
+						ParentTransform = FTransform3f(ParentRotation, ParentLocation).ToMatrixNoScale();
+					}
+
+					for (int32 Idx = 0; Idx < 4; Idx++) { Swap(ParentTransform.M[1][Idx], ParentTransform.M[2][Idx]); }
+
+					for (int32 Idx = 0; Idx < 3; Idx++)
+					{
+						const float Y = ParentTransform.M[1][Idx];
+						const float Z = ParentTransform.M[2][Idx];
+						ParentTransform.M[2][Idx] = +Y;
+						ParentTransform.M[1][Idx] = -Z;
+					}
+
+					const FMatrix44f LocalTransform = BoneTransform * ParentTransform.Inverse();
+					const FQuat4f LocalRotation = UE::AnimDatabase::Math::MakeQuatFromMatrixRightHanded(LocalTransform).GetNormalized();
+					const FVector3f LocalLocation = FVector3f(LocalTransform.M[3][0], LocalTransform.M[3][1], LocalTransform.M[3][2]);
+
+					if (BoneIdx == 0)
+					{
+						PoseData.RootData.RootLocations[LocalFrameIdx] = (FVector)LocalLocation;
+						PoseData.RootData.RootRotations[LocalFrameIdx] = LocalRotation;
+					}
+					else
+					{
+						PoseData.LocalBoneData.BoneLocations[LocalFrameIdx][DatabaseBoneIdx] = LocalLocation;
+						PoseData.LocalBoneData.BoneRotations[LocalFrameIdx][DatabaseBoneIdx] = LocalRotation;
+					}
+				}
+			}
+
+
+			TArray<int32> UsedBones;
+			UE::AnimDatabase::PoseData::ComputeUsedBones(UsedBones, PoseData.LocalBoneData.ConstView(), Database->Skeleton->GetReferenceSkeleton().GetRefBonePose());
+
+			// Create Asset
+
+			FString AssetName = File;
+			AssetName.RemoveFromEnd(TEXT(".bvh"));
+
+			FString PathPackageName;
+			if (!FPackageName::TryConvertFilenameToLongPackageName(AssetPath.Path, PathPackageName))
+			{
+				UE_LOGFMT(LogAnimDatabase, Error, "DatabaseImportFromBVH: Unable to convert path '{Path}' to package name.", AssetPath.Path);
+				continue;
+			}
+
+			const FString PackageName = PathPackageName + TEXT("/") + AssetName;
+
+			UPackage* Package = CreatePackage(*PackageName);
+			if (!Package)
+			{
+				UE_LOGFMT(LogAnimDatabase, Error, "DatabaseImportFromBVH: Unable to create package '{PackageName}'.", PackageName);
+				continue;
+			}
+			Package->FullyLoad();
+
+			// Find or Create Anim Sequence
+
+			UAnimSequence* NewAnimSequence = FindObject<UAnimSequence>(nullptr, PackageName + TEXT(".") + AssetName);
+			if (!NewAnimSequence)
+			{
+				NewAnimSequence = NewObject<UAnimSequence>(Package, FName(*AssetName), RF_Public | RF_Standalone);
+			}
+
+			if (!NewAnimSequence)
+			{
+				UE_LOGFMT(LogAnimDatabase, Error, "DatabaseImportFromBVH: Unable to create new AnimSequence.");
+				continue;
+			}
+
+			// Setup Sequence
+
+			const FFrameRate AnimFrameRate = FFrameRate(FMath::RoundToInt(1.0f / FMath::Max(FrameTime, UE_SMALL_NUMBER)), 1);
+			if (AnimFrameRate != Database->GetFrameRate())
+			{
+				UE_LOGFMT(LogAnimDatabase, Warning, "DatabaseImportFromBVH: Animation FrameRate ({0}/{1}) does not match database ({2}/{0}).", 
+					AnimFrameRate.Numerator, AnimFrameRate.Denominator, Database->GetFrameRate().Numerator, Database->GetFrameRate().Denominator);
+			}
+
+			const bool bShouldTransact = false;
+
+			NewAnimSequence->ResetAnimation();
+			NewAnimSequence->SetSkeleton(Database->Skeleton);
+			NewAnimSequence->bForceRootLock = true;
+			NewAnimSequence->bEnableRootMotion = true;
+
+			TScriptInterface<IAnimationDataController> Controller = NewAnimSequence->GetDataModel()->GetController();
+			Controller->OpenBracket(LOCTEXT("CreateAnimSequence", "Creating Animation Sequence"), bShouldTransact);
+			Controller->InitializeModel();
+			Controller->SetFrameRate(AnimFrameRate, bShouldTransact);
+			Controller->SetNumberOfFrames(FrameNum - 1, bShouldTransact);
+
+			// Key initial bone data to reference pose
+
+			TArray<FVector> AssetBonePositions;
+			TArray<FQuat> AssetBoneRotations;
+			TArray<FVector> AssetBoneScales;
+			AssetBonePositions.SetNumUninitialized(FrameNum);
+			AssetBoneRotations.SetNumUninitialized(FrameNum);
+			AssetBoneScales.SetNumUninitialized(FrameNum);
+
+			const FReferenceSkeleton ReferenceSkeleton = Database->GetSkeleton()->GetReferenceSkeleton();
+
+			for (int32 LocalBoneIdx : UsedBones)
+			{
+				const FTransform RefTransform = ReferenceSkeleton.GetRefBonePose()[LocalBoneIdx];
+				const FName BoneName = ReferenceSkeleton.GetBoneName(LocalBoneIdx);
+				check(BoneName != NAME_None);
+
+				if (!Controller->GetModel()->IsValidBoneTrackName(BoneName))
+				{
+					Controller->AddBoneCurve(BoneName, bShouldTransact);
+				}
+
+				UE::Learning::Array::Set<1, FVector>(AssetBonePositions, RefTransform.GetLocation());
+				UE::Learning::Array::Set<1, FQuat>(AssetBoneRotations, RefTransform.GetRotation());
+				UE::Learning::Array::Set<1, FVector>(AssetBoneScales, RefTransform.GetScale3D());
+
+				Controller->SetBoneTrackKeys(BoneName, AssetBonePositions, AssetBoneRotations, AssetBoneScales, bShouldTransact);
+			}
+
+			// Set it to the actual animation data
+
+			SetAnimSequencePoseData(NewAnimSequence, Controller, 0, PoseData.ConstView(), UsedBones, bShouldTransact);
+
+			Controller->NotifyPopulated();
+			Controller->CloseBracket(bShouldTransact);
+
+			NewAnimSequence->RefreshCacheData();
+
+			// Inform the asset registry about the creation.
+			Package->MarkPackageDirty();
+			FAssetRegistryModule::AssetCreated(NewAnimSequence);
+
+			// Add to Database
+			if (bAddToDatabase)
+			{
+				Database->Entries.AddUnique(NewAnimSequence);
+				Database->Modify();
+			}
 		}
 	}
 }
@@ -5122,6 +5579,9 @@ void UAnimDatabaseLibrary::DatabaseExportAsAnimSequences(UAnimDatabase* Database
 	UE::AnimDatabase::FPoseData PoseData;
 	UE::AnimDatabase::FPoseGlobalBoneData GlobalPoseData;
 
+	FScopedSlowTask SlowTask(FrameRanges.FrameRangeSet->GetTotalRangeNum(), LOCTEXT("ExportingToAnimSequences", "Exporting to Anim Sequences"));
+	SlowTask.MakeDialog(true);
+
 	const int32 EntryNum = FrameRanges.FrameRangeSet->GetEntryNum();
 
 	for (int32 EntryIdx = 0; EntryIdx < EntryNum; EntryIdx++)
@@ -5135,6 +5595,9 @@ void UAnimDatabaseLibrary::DatabaseExportAsAnimSequences(UAnimDatabase* Database
 		{
 			for (int32 RangeIdx = 0; RangeIdx < RangeNum; RangeIdx++)
 			{
+				SlowTask.EnterProgressFrame();
+				if (SlowTask.ShouldCancel()) { return; }
+
 				const int32 FrameStart = FrameRanges.FrameRangeSet->GetEntryRangeStart(EntryIdx, RangeIdx);
 				const int32 FrameNum = FrameRanges.FrameRangeSet->GetEntryRangeLength(EntryIdx, RangeIdx);
 
@@ -6009,6 +6472,97 @@ void UAnimDatabaseLibrary::DatabaseSetToReferencePose(UAnimDatabase* Database, c
 		}, bShouldTransact);
 }
 
+void UAnimDatabaseLibrary::DatabaseRemoveFootGroundPenetration(UAnimDatabase* Database, const FAnimDatabaseFrameRanges& FrameRanges, const FName LeftToeBone, const FName RightToeBone, const float ToeBoneLength, const float PelvisHeightAdjustment, const FVector LeftKneeSideVector, const FVector RightKneeSideVector, const FVector LeftToeForwardVector, const FVector RightToeForwardVector, const bool bShouldTransact)
+{
+	if (!Database || !Database->GetSkeleton()) { return; }
+
+	const int32 LeftToeBoneIndex = Database->FindBoneIndex(LeftToeBone);
+	const int32 LeftHeelBoneIndex = Database->GetBoneParent(LeftToeBoneIndex);
+	const int32 LeftKneeBoneIndex = Database->GetBoneParent(LeftHeelBoneIndex);
+	const int32 LeftHipBoneIndex = Database->GetBoneParent(LeftKneeBoneIndex);
+	const int32 LeftPelvisBoneIndex = Database->GetBoneParent(LeftHipBoneIndex);
+
+	const int32 RightToeBoneIndex = Database->FindBoneIndex(RightToeBone);
+	const int32 RightHeelBoneIndex = Database->GetBoneParent(RightToeBoneIndex);
+	const int32 RightKneeBoneIndex = Database->GetBoneParent(RightHeelBoneIndex);
+	const int32 RightHipBoneIndex = Database->GetBoneParent(RightKneeBoneIndex);
+	const int32 RightPelvisBoneIndex = Database->GetBoneParent(RightHipBoneIndex);
+
+	if (LeftToeBoneIndex == INDEX_NONE || RightToeBoneIndex == INDEX_NONE ||
+		LeftHeelBoneIndex == INDEX_NONE || RightHeelBoneIndex == INDEX_NONE ||
+		LeftKneeBoneIndex == INDEX_NONE || RightKneeBoneIndex == INDEX_NONE ||
+		LeftHipBoneIndex == INDEX_NONE || RightHipBoneIndex == INDEX_NONE ||
+		LeftPelvisBoneIndex == INDEX_NONE || RightPelvisBoneIndex == INDEX_NONE)
+	{
+		return; 
+	}
+
+	UE::AnimDatabase::PoseData::ContactPinning::FContactPinningSettings ContactSettings;
+	ContactSettings.ContactHeight = Database->GetSkeleton()->GetReferenceSkeleton().GetBoneAbsoluteTransform(LeftToeBoneIndex).GetLocation().Z;
+	ContactSettings.MinToeHeight = Database->GetSkeleton()->GetReferenceSkeleton().GetBoneAbsoluteTransform(LeftToeBoneIndex).GetLocation().Z;
+	ContactSettings.MinAnkleHeight = Database->GetSkeleton()->GetReferenceSkeleton().GetBoneAbsoluteTransform(LeftHeelBoneIndex).GetLocation().Z;
+	ContactSettings.ToeLength = ToeBoneLength;
+	ContactSettings.HyperExtensionLimit = 0.0f;
+
+	TLearningArray<1, int32> BoneParents;
+	BoneParents.SetNumUninitialized({ Database->GetBoneNum() });
+	Database->GetBoneParentsToArrayView(BoneParents);
+
+	DatabaseModifyPoseData(Database, FrameRanges, [
+		Database, 
+		PelvisHeightAdjustment,
+		LeftPelvisBoneIndex,
+		RightPelvisBoneIndex,
+		LeftHipBoneIndex,
+		RightHipBoneIndex,
+		LeftKneeBoneIndex,
+		RightKneeBoneIndex,
+		LeftHeelBoneIndex,
+		RightHeelBoneIndex,
+		LeftToeBoneIndex,
+		RightToeBoneIndex,
+		LeftKneeSideVector,
+		RightKneeSideVector,
+		LeftToeForwardVector,
+		RightToeForwardVector,
+		&BoneParents,
+		&ContactSettings
+	](const UE::AnimDatabase::FPoseDataView& InOutPoseData)
+		{
+			const int32 FrameNum = InOutPoseData.GetFrameNum();
+			const int32 BoneNum = InOutPoseData.GetBoneNum();
+
+			for (int32 FrameIdx = 0; FrameIdx < FrameNum; FrameIdx++)
+			{
+				InOutPoseData.LocalBoneData.BoneLocations[FrameIdx][LeftPelvisBoneIndex] += FVector3f(0.0f, 0.0f, PelvisHeightAdjustment);
+			}
+		
+			UE::AnimDatabase::FPoseGlobalBoneData GlobalData;
+			GlobalData.Resize(FrameNum, BoneNum);
+
+			UE::AnimDatabase::PoseData::ForwardKinematics(
+				GlobalData.View(), 
+				InOutPoseData.LocalBoneData.ConstView(), 
+				InOutPoseData.RootData.ConstView(),
+				BoneParents);
+
+			UE::AnimDatabase::PoseData::ContactPinning::RemoveFootGroundPenetration(
+				GlobalData.View(),
+				InOutPoseData.LocalBoneData,
+				InOutPoseData.RootData.ConstView(),
+				BoneParents,
+				{ LeftPelvisBoneIndex, RightPelvisBoneIndex },
+				{ LeftHipBoneIndex, RightHipBoneIndex },
+				{ LeftKneeBoneIndex, RightKneeBoneIndex },
+				{ LeftHeelBoneIndex, RightHeelBoneIndex },
+				{ LeftToeBoneIndex, RightToeBoneIndex },
+				{ (FVector3f)(LeftKneeSideVector), (FVector3f)RightKneeSideVector },
+				{ (FVector3f)(LeftToeForwardVector), (FVector3f)RightToeForwardVector },
+				ContactSettings);
+
+		}, bShouldTransact);
+}
+
 void UAnimDatabaseFunction_ExportBVH::Run_Implementation(UAnimDatabase* InDatabase, const FAnimDatabaseFrameRanges& InFrameRanges, const bool bShouldTransact)
 {
 	if (!ExportDirectory.Path.IsEmpty())
@@ -6017,6 +6571,14 @@ void UAnimDatabaseFunction_ExportBVH::Run_Implementation(UAnimDatabase* InDataba
 	}
 }
 
+void UAnimDatabaseFunction_ImportBVH::Run_Implementation(UAnimDatabase* InDatabase, const FAnimDatabaseFrameRanges& InFrameRanges, const bool bShouldTransact)
+{
+	if (!ImportDirectory.Path.IsEmpty() && !AssetDirectory.Path.IsEmpty())
+	{
+		UAnimDatabaseLibrary::DatabaseImportFromBVH(InDatabase, ImportDirectory, AssetDirectory, bRemoveRootAxesChangeRotation, bAddToDatabase, bIgnoreMirrored);
+	}
+}
+
 void UAnimDatabaseFunction_ExportNotifies::Run_Implementation(UAnimDatabase* InDatabase, const FAnimDatabaseFrameRanges& InFrameRanges, const bool bShouldTransact)
 {
 	if (!ExportDirectory.Path.IsEmpty())
@@ -6290,6 +6852,11 @@ void UAnimDatabaseFunction_AdjustNeckRotation::Run_Implementation(UAnimDatabase*
 		bShouldTransact);
 }
 
+void UAnimDatabaseFunction_RemoveFootGroundPenetration::Run_Implementation(UAnimDatabase* InDatabase, const FAnimDatabaseFrameRanges& InFrameRanges, const bool bShouldTransact)
+{
+	UAnimDatabaseLibrary::DatabaseRemoveFootGroundPenetration(InDatabase, InFrameRanges, LeftToeBone, RightToeBone, ToeBoneLength, PelvisHeightAdjustment, LeftKneeSideVector, RightKneeSideVector, LeftToeForwardVector, RightToeForwardVector, bShouldTransact);
+}
+
 void UAnimDatabaseFunction_Statistics::Run_Implementation(UAnimDatabase* InDatabase, const FAnimDatabaseFrameRanges& InFrameRanges, const bool bShouldTransact)
 {
 	UAnimDatabaseFrameRangesLibrary::FrameRangesStatistics(Statistics, InDatabase, InFrameRanges);
@@ -6401,6 +6968,67 @@ void UAnimDatabaseDebugDraw_ContactCurves::DrawDebug_Implementation(
 	}
 }
 
+void UAnimDatabaseDebugDraw_Multi::InitializeDrawDebug_Implementation(const UAnimDatabase* InDatabase, const FAnimDatabaseFrameRanges& InFrameRanges)
+{
+	for (TObjectPtr<UAnimDatabaseDebugDraw>& DrawerObject : Drawers)
+	{
+		if (DrawerObject) { DrawerObject->InitializeDrawDebug(InDatabase, InFrameRanges); }
+	}
+}
+
+void UAnimDatabaseDebugDraw_Multi::DrawDebug_Implementation(
+	const FDebugDrawer& Drawer,
+	const FAnimDatabasePoseState& InPoseState,
+	const UAnimDatabase* InDatabase,
+	const FAnimDatabaseFrameRanges& InFrameRanges,
+	const FTransform& RangeViewportTransform,
+	const int32 SequenceIdx,
+	const float SequenceTime,
+	const int32 RangeStart,
+	const int32 RangeLength,
+	const FLinearColor& IdentifierColor)
+{
+	for (TObjectPtr<UAnimDatabaseDebugDraw>& DrawerObject : Drawers)
+	{
+		if (DrawerObject)
+		{
+			DrawerObject->DrawDebug(
+				Drawer,
+				InPoseState,
+				InDatabase, 
+				InFrameRanges,
+				RangeViewportTransform,
+				SequenceIdx,
+				SequenceTime,
+				RangeStart,
+				RangeLength,
+				IdentifierColor);
+		}
+	}
+}
+
+void UAnimDatabaseDebugDraw_RootOrientation::DrawDebug_Implementation(
+	const FDebugDrawer& Drawer,
+	const FAnimDatabasePoseState& InPoseState,
+	const UAnimDatabase* InDatabase,
+	const FAnimDatabaseFrameRanges& InFrameRanges,
+	const FTransform& RangeViewportTransform,
+	const int32 SequenceIdx,
+	const float SequenceTime,
+	const int32 RangeStart,
+	const int32 RangeLength,
+	const FLinearColor& IdentifierColor)
+{
+	UDrawDebugLibrary::DrawDebugMoverOrientation(
+		Drawer,
+		UAnimDatabasePoseStateLibrary::PoseStateRootLocation(InPoseState) + FVector(0.0, 0.0, FloorOffset),
+		UAnimDatabasePoseStateLibrary::PoseStateRootRotation(InPoseState),
+		UDrawDebugLibrary::MakeDrawDebugLineStyleFromColorAndThickness(IdentifierColor, 1.0f),
+		true,
+		ForwardVector,
+		Radius);
+}
+
 void UAnimDatabaseDebugDraw_RootVelocity::InitializeDrawDebug_Implementation(const UAnimDatabase* InDatabase, const FAnimDatabaseFrameRanges& InFrameRanges)
 {
 	RootLinearVelocity = UAnimDatabaseFrameAttributeLibrary::FrameAttributeLength(UAnimDatabaseFrameAttributeLibrary::MakeRootLinearVelocityFrameAttribute(InDatabase, InFrameRanges));
diff --git a/Engine/Plugins/Experimental/Animation/AnimDatabase/Source/AnimDatabase/Private/AnimDatabasePose.cpp b/Engine/Plugins/Experimental/Animation/AnimDatabase/Source/AnimDatabase/Private/AnimDatabasePose.cpp
--- a/Engine/Plugins/Experimental/Animation/AnimDatabase/Source/AnimDatabase/Private/AnimDatabasePose.cpp
+++ b/Engine/Plugins/Experimental/Animation/AnimDatabase/Source/AnimDatabase/Private/AnimDatabasePose.cpp
@@ -49,6 +49,18 @@ namespace UE::AnimDatabase
 		return ConstView().Slice(FrameStart, FrameNum);
 	}
 
+	FPoseRootDataConstView FPoseRootDataView::ConstView() const
+	{
+		FPoseRootDataConstView View;
+		View.RootLocations = RootLocations;
+		View.RootRotations = RootRotations;
+		View.RootScales = RootScales;
+		View.RootLinearVelocities = RootLinearVelocities;
+		View.RootAngularVelocities = RootAngularVelocities;
+		View.RootScalarVelocities = RootScalarVelocities;
+		return View;
+	}
+
 	FPoseRootDataView FPoseRootDataView::Slice(const int32 FrameStart, const int32 FrameNum) const
 	{
 		FPoseRootDataView View;
@@ -3591,6 +3603,7 @@ namespace UE::AnimDatabase
 				const TArrayView<const int32> AnkleBoneIndices,
 				const TArrayView<const int32> ToeBoneIndices,
 				const TArrayView<const FVector3f> KneeSideVectors,
+				const TArrayView<const FVector3f> ToeForwardVectors,
 				const float DeltaTime,
 				const FContactPinningSettings& Settings)
 			{
@@ -3605,6 +3618,7 @@ namespace UE::AnimDatabase
 				check(InOutContactState.GetBoneNum() == AnkleBoneIndices.Num());
 				check(InOutContactState.GetBoneNum() == ToeBoneIndices.Num());
 				check(InOutContactState.GetBoneNum() == KneeSideVectors.Num());
+				check(InOutContactState.GetBoneNum() == ToeForwardVectors.Num());
 
 				const int32 FrameNum = InOutContactState.GetFrameNum();
 				const int32 BoneNum = InOutContactState.GetBoneNum();
@@ -3745,7 +3759,133 @@ namespace UE::AnimDatabase
 					for (int32 BoneIdx = 0; BoneIdx < BoneNum; BoneIdx++)
 					{
 						const FVector ToeEndCurr =
-							(FVector)InOutPoseGlobalBoneData.BoneRotations[FrameIdx][ToeBoneIndices[BoneIdx]].RotateVector(FVector3f(-Settings.ToeLength, 0.0f, 0.0f)) +
+							(FVector)InOutPoseGlobalBoneData.BoneRotations[FrameIdx][ToeBoneIndices[BoneIdx]].RotateVector(Settings.ToeLength * ToeForwardVectors[BoneIdx]) +
+							InOutPoseGlobalBoneData.BoneLocations[FrameIdx][ToeBoneIndices[BoneIdx]];
+
+						FVector ToeEndTarg = ToeEndCurr;
+						ToeEndTarg.Z = FMath::Max(ToeEndTarg.Z, Settings.MinToeHeight);
+
+						ContactLookAtIK(
+							InOutPoseLocalBoneData.BoneRotations[FrameIdx][ToeBoneIndices[BoneIdx]],
+							InOutPoseGlobalBoneData.BoneRotations[FrameIdx][AnkleBoneIndices[BoneIdx]],
+							InOutPoseGlobalBoneData.BoneRotations[FrameIdx][ToeBoneIndices[BoneIdx]],
+							InOutPoseGlobalBoneData.BoneLocations[FrameIdx][ToeBoneIndices[BoneIdx]],
+							ToeEndCurr,
+							ToeEndTarg);
+					}
+				}
+			}
+
+			void RemoveFootGroundPenetration(
+				const FPoseGlobalBoneDataView& InOutPoseGlobalBoneData,
+				const FPoseLocalBoneDataView& InOutPoseLocalBoneData,
+				const FPoseRootDataConstView& InPoseRootData,
+				const TArrayView<const int32> InBoneParents,
+				const TArrayView<const int32> PelvisBoneIndices,
+				const TArrayView<const int32> HipBoneIndices,
+				const TArrayView<const int32> KneeBoneIndices,
+				const TArrayView<const int32> AnkleBoneIndices,
+				const TArrayView<const int32> ToeBoneIndices,
+				const TArrayView<const FVector3f> KneeSideVectors,
+				const TArrayView<const FVector3f> ToeForwardVectors,
+				const FContactPinningSettings& Settings)
+			{
+				check(InOutPoseGlobalBoneData.GetFrameNum() == InOutPoseLocalBoneData.GetFrameNum());
+				check(InOutPoseGlobalBoneData.GetFrameNum() == InPoseRootData.GetFrameNum());
+				check(PelvisBoneIndices.Num() == HipBoneIndices.Num());
+				check(PelvisBoneIndices.Num() == KneeBoneIndices.Num());
+				check(PelvisBoneIndices.Num() == AnkleBoneIndices.Num());
+				check(PelvisBoneIndices.Num() == ToeBoneIndices.Num());
+				check(PelvisBoneIndices.Num() == KneeSideVectors.Num());
+				check(PelvisBoneIndices.Num() == ToeForwardVectors.Num());
+
+				const int32 FrameNum = InOutPoseGlobalBoneData.GetFrameNum();
+				const int32 BoneNum = PelvisBoneIndices.Num();
+
+				TArray<int32, TInlineAllocator<32>> RequiredBones;
+				TArray<int32, TInlineAllocator<32>> Lhs;
+				TArray<int32, TInlineAllocator<32>> Rhs;
+				for (int32 BoneIdx = 0; BoneIdx < BoneNum; BoneIdx++)
+				{
+					Lhs = RequiredBones;
+					Rhs.SetNumUninitialized(Math::BoneAscendantsInclusiveNum(ToeBoneIndices[BoneIdx], InBoneParents));
+					Math::BoneAscendantsInclusive(Rhs, ToeBoneIndices[BoneIdx], InBoneParents);
+					RequiredBones.SetNumUninitialized(Math::BoneUnionNum(Lhs, Rhs));
+					Math::BoneUnion(RequiredBones, Lhs, Rhs);
+				}
+
+				PoseData::ForwardKinematicsPartial(
+					InOutPoseGlobalBoneData,
+					InOutPoseLocalBoneData.ConstView(),
+					InPoseRootData,
+					InBoneParents,
+					UE::Learning::FIndexSet(0, InBoneParents.Num()),
+					RequiredBones);
+
+				for (int32 FrameIdx = 0; FrameIdx < FrameNum; FrameIdx++)
+				{
+					for (int32 BoneIdx = 0; BoneIdx < BoneNum; BoneIdx++)
+					{
+						FVector ToePosition = InOutPoseGlobalBoneData.BoneLocations[FrameIdx][ToeBoneIndices[BoneIdx]];
+						ToePosition.Z = FMath::Max(ToePosition.Z, Settings.MinToeHeight);
+
+						FVector AnklePosition = ToePosition + (InOutPoseGlobalBoneData.BoneLocations[FrameIdx][AnkleBoneIndices[BoneIdx]] - InOutPoseGlobalBoneData.BoneLocations[FrameIdx][ToeBoneIndices[BoneIdx]]);
+						AnklePosition.Z = FMath::Max(AnklePosition.Z, Settings.MinAnkleHeight);
+
+						ContactTwoBoneIK(
+							InOutPoseLocalBoneData.BoneRotations[FrameIdx][HipBoneIndices[BoneIdx]],
+							InOutPoseLocalBoneData.BoneRotations[FrameIdx][KneeBoneIndices[BoneIdx]],
+							InOutPoseGlobalBoneData.BoneLocations[FrameIdx][HipBoneIndices[BoneIdx]],
+							InOutPoseGlobalBoneData.BoneLocations[FrameIdx][KneeBoneIndices[BoneIdx]],
+							InOutPoseGlobalBoneData.BoneLocations[FrameIdx][AnkleBoneIndices[BoneIdx]],
+							AnklePosition,
+							InOutPoseGlobalBoneData.BoneRotations[FrameIdx][KneeBoneIndices[BoneIdx]].RotateVector(KneeSideVectors[BoneIdx]),
+							InOutPoseGlobalBoneData.BoneRotations[FrameIdx][HipBoneIndices[BoneIdx]],
+							InOutPoseGlobalBoneData.BoneRotations[FrameIdx][KneeBoneIndices[BoneIdx]],
+							InOutPoseGlobalBoneData.BoneRotations[FrameIdx][PelvisBoneIndices[BoneIdx]],
+							Settings.HyperExtensionLimit);
+					}
+				}
+
+				PoseData::ForwardKinematicsPartial(
+					InOutPoseGlobalBoneData,
+					InOutPoseLocalBoneData.ConstView(),
+					InPoseRootData,
+					InBoneParents,
+					UE::Learning::FIndexSet(0, InBoneParents.Num()),
+					RequiredBones);
+
+				for (int32 FrameIdx = 0; FrameIdx < FrameNum; FrameIdx++)
+				{
+					for (int32 BoneIdx = 0; BoneIdx < BoneNum; BoneIdx++)
+					{
+						FVector ToePosition = InOutPoseGlobalBoneData.BoneLocations[FrameIdx][ToeBoneIndices[BoneIdx]];
+						ToePosition.Z = FMath::Max(ToePosition.Z, Settings.MinToeHeight);
+
+						ContactLookAtIK(
+							InOutPoseLocalBoneData.BoneRotations[FrameIdx][AnkleBoneIndices[BoneIdx]],
+							InOutPoseGlobalBoneData.BoneRotations[FrameIdx][KneeBoneIndices[BoneIdx]],
+							InOutPoseGlobalBoneData.BoneRotations[FrameIdx][AnkleBoneIndices[BoneIdx]],
+							InOutPoseGlobalBoneData.BoneLocations[FrameIdx][AnkleBoneIndices[BoneIdx]],
+							InOutPoseGlobalBoneData.BoneLocations[FrameIdx][ToeBoneIndices[BoneIdx]],
+							ToePosition);
+					}
+				}
+
+				PoseData::ForwardKinematicsPartial(
+					InOutPoseGlobalBoneData,
+					InOutPoseLocalBoneData.ConstView(),
+					InPoseRootData,
+					InBoneParents,
+					UE::Learning::FIndexSet(0, InBoneParents.Num()),
+					RequiredBones);
+
+				for (int32 FrameIdx = 0; FrameIdx < FrameNum; FrameIdx++)
+				{
+					for (int32 BoneIdx = 0; BoneIdx < BoneNum; BoneIdx++)
+					{
+						const FVector ToeEndCurr =
+							(FVector)InOutPoseGlobalBoneData.BoneRotations[FrameIdx][ToeBoneIndices[BoneIdx]].RotateVector(Settings.ToeLength * ToeForwardVectors[BoneIdx]) +
 							InOutPoseGlobalBoneData.BoneLocations[FrameIdx][ToeBoneIndices[BoneIdx]];
 
 						FVector ToeEndTarg = ToeEndCurr;
diff --git a/Engine/Plugins/Experimental/Animation/AnimDatabase/Source/AnimDatabase/Public/AnimDatabase.h b/Engine/Plugins/Experimental/Animation/AnimDatabase/Source/AnimDatabase/Public/AnimDatabase.h
--- a/Engine/Plugins/Experimental/Animation/AnimDatabase/Source/AnimDatabase/Public/AnimDatabase.h
+++ b/Engine/Plugins/Experimental/Animation/AnimDatabase/Source/AnimDatabase/Public/AnimDatabase.h
@@ -493,7 +493,7 @@ DECLARE_MULTICAST_DELEGATE(FAnimDatabaseOnSkeletonChanged);
  * The number of "Sequences" in this database (as returned by GetSequenceNum()) will be two times the number of entries in "Entries" when the mirror
  * data table is present and valid, otherwise it will be the same as the number of entries.
  */
-UCLASS(Blueprintable)
+UCLASS(MinimalAPI, BlueprintType, Category = "Animation", meta = (DisplayName = "Animation Database"))
 class UAnimDatabase : public UDataAsset
 {
 	GENERATED_BODY()
@@ -535,7 +535,9 @@ public:
 
 	/** Callback to fire when the skeleton is changed */
 	FAnimDatabaseOnSkeletonChanged OnSkeletonChanged;
+#endif
 
+#if WITH_EDITOR
 	/** We use this to ensure that entries added have a matching skeleton */
 	UE_API virtual void PostEditChangeProperty(struct FPropertyChangedEvent& Event) override;
 #endif
@@ -941,6 +943,10 @@ public:
 	UFUNCTION(BlueprintCallable, Category = "AnimDatabase", meta = (DevelopmentOnly, RelativePath))
 	static UE_API void DatabaseExportAsBVH(UAnimDatabase* Database, const FAnimDatabaseFrameRanges& FrameRanges, const FDirectoryPath& Path, const bool bAddRootAxesChangeRotation = true, const EAnimDatabaseRotationOrder RotationOrder = EAnimDatabaseRotationOrder::XYZ);
 
+	/** Imports all the BVH files from the given ImportPath directory into the AssetPath directory and adds them to the database */
+	UFUNCTION(BlueprintCallable, Category = "AnimDatabase", meta = (DevelopmentOnly, RelativePath))
+	static UE_API void DatabaseImportFromBVH(UAnimDatabase* Database, const FDirectoryPath& ImportPath, const FDirectoryPath& AssetPath, const bool bRemoveRootAxesChangeRotation = true, const bool bAddToDatabase = true, const bool bIgnoreMirrored = false);
+
 	/** Exports the given frame ranges in the database as AnimSequences into the provided path */
 	UFUNCTION(BlueprintCallable, Category = "AnimDatabase", meta = (DevelopmentOnly, RelativePath))
 	static UE_API void DatabaseExportAsAnimSequences(UAnimDatabase* Database, const FAnimDatabaseFrameRanges& FrameRanges, const FDirectoryPath& Path, const FString& AssetNameFormatString, const bool bExportMirrored, const bool bShouldTransact = false);
@@ -1008,6 +1014,9 @@ public:
 	UFUNCTION(BlueprintCallable, Category = "AnimDatabase", meta = (DevelopmentOnly))
 	static UE_API void DatabaseSetToReferencePose(UAnimDatabase* Database, const FAnimDatabaseFrameRanges& FrameRanges, const bool bApplyToRoot, const TArray<int32>& IgnoreBoneIndices, const bool bShouldTransact = false);
 
+	/** Removes foot ground penetration in the database */
+	UFUNCTION(BlueprintCallable, Category = "AnimDatabase", meta = (DevelopmentOnly))
+	static UE_API void DatabaseRemoveFootGroundPenetration(UAnimDatabase* Database, const FAnimDatabaseFrameRanges& FrameRanges, const FName LeftToeBone, const FName RightToeBone, const float ToeBoneLength, const float PelvisHeightAdjustment = -1.0f, const FVector LeftKneeSideVector = FVector(0.0f, 0.0f, +1.0f), const FVector RightKneeSideVector = FVector(0.0f, 0.0f, +1.0f), const FVector LeftToeForwardVector = FVector(-1.0f, 0.0f, 0.0f), const FVector RightToeForwardVector = FVector(-1.0f, 0.0f, 0.0f), const bool bShouldTransact = false);
 
 private:
 
@@ -1047,6 +1056,43 @@ public:
 #endif
 };
 
+/** Function for importing from BVH  */
+UCLASS(BlueprintType, Blueprintable, meta = (DisplayName = "Import from BVH"))
+class UAnimDatabaseFunction_ImportBVH : public UAnimDatabaseFunction
+{
+	GENERATED_BODY()
+
+public:
+
+	/** Directory to import BVH files from. */
+	UPROPERTY(EditAnywhere, Category = "Settings", meta = (RelativePath))
+	FDirectoryPath ImportDirectory;
+
+	/** Directory to write AnimSequences to. */
+	UPROPERTY(EditAnywhere, Category = "Settings", meta = (RelativePath))
+	FDirectoryPath AssetDirectory;
+
+	/** If to remove the appropriate root rotation to account for the axes switch. */
+	UPROPERTY(EditAnywhere, Category = "Settings")
+	bool bRemoveRootAxesChangeRotation = true;
+
+	/** Ignores BVH files with the _mirrored suffix */
+	UPROPERTY(EditAnywhere, Category = "Settings")
+	bool bIgnoreMirrored = false;
+
+	/** If to add the imported assets to the database. */
+	UPROPERTY(EditAnywhere, Category = "Settings")
+	bool bAddToDatabase = true;
+
+#if WITH_EDITOR
+
+public:
+
+	UE_API virtual void Run_Implementation(UAnimDatabase* InDatabase, const FAnimDatabaseFrameRanges& InFrameRanges, const bool bShouldTransact) override;
+
+#endif
+};
+
 /** Function for exporting all of the AnimNotifies and AnimNotifyStates from the given animations  */
 UCLASS(BlueprintType, Blueprintable, meta = (DisplayName = "Export Notifies"))
 class UAnimDatabaseFunction_ExportNotifies : public UAnimDatabaseFunction
@@ -1667,7 +1713,7 @@ public:
 
 	/** Rotation to apply to the head bone */
 	UPROPERTY(EditAnywhere, Category = "Settings", meta = (ForceUnits = "deg"))
-	float HeadRotation = 16.0f;
+	float HeadRotation = 18.0f;
 
 	/** Local axis to apply rotation on */
 	UPROPERTY(EditAnywhere, Category = "Settings")
@@ -1684,6 +1730,59 @@ public:
 #endif
 };
 
+/** Function to remove foot ground penetration */
+UCLASS(BlueprintType, Blueprintable, meta = (DisplayName = "Remove Foot Ground Penetration"))
+class UAnimDatabaseFunction_RemoveFootGroundPenetration : public UAnimDatabaseFunction
+{
+	GENERATED_BODY()
+
+#if WITH_EDITORONLY_DATA
+
+public:
+
+	/** Name of the left toe bone */
+	UPROPERTY(EditAnywhere, Category = "Settings")
+	FName LeftToeBone = TEXT("ball_l");
+
+	/** Name of the right toe bone */
+	UPROPERTY(EditAnywhere, Category = "Settings")
+	FName RightToeBone = TEXT("ball_r");
+
+	/** Length of the toe joint */
+	UPROPERTY(EditAnywhere, Category = "Settings", meta=(ForceUnits="cm"))
+	float ToeBoneLength = 7.0f;
+
+	/** Side Vector of the Left Knee */
+	UPROPERTY(EditAnywhere, Category = "Settings")
+	FVector LeftKneeSideVector = FVector(0.0f, 0.0f, +1.0f);
+
+	/** Side Vector of the Right Knee */
+	UPROPERTY(EditAnywhere, Category = "Settings")
+	FVector RightKneeSideVector = FVector(0.0f, 0.0f, +1.0f);
+
+	/** Forward Vector of the Left Toe */
+	UPROPERTY(EditAnywhere, Category = "Settings")
+	FVector LeftToeForwardVector = FVector(-1.0f, 0.0f, 0.0f);
+
+	/** Forward Vector of the Right Toe */
+	UPROPERTY(EditAnywhere, Category = "Settings")
+	FVector RightToeForwardVector = FVector(-1.0f, 0.0f, 0.0f);
+
+	/** Adjustment to apply to the pelvis height */
+	UPROPERTY(EditAnywhere, Category = "Settings", meta = (ForceUnits = "cm"))
+	float PelvisHeightAdjustment = -1.0f;
+
+#endif
+
+#if WITH_EDITOR
+
+public:
+
+	UE_API virtual void Run_Implementation(UAnimDatabase* InDatabase, const FAnimDatabaseFrameRanges& InFrameRanges, const bool bShouldTransact) override;
+
+#endif
+};
+
 /** Computes various statistics about the provided ranges */
 UCLASS(BlueprintType, Blueprintable, meta = (DisplayName = "Statistics"))
 class UAnimDatabaseFunction_Statistics : public UAnimDatabaseFunction
@@ -1922,6 +2021,86 @@ public:
 };
 
 
+/** Debug Draw class for multiple other debug drawers */
+UCLASS(BlueprintType, Blueprintable, meta = (DisplayName = "Multi"))
+class UAnimDatabaseDebugDraw_Multi : public UAnimDatabaseDebugDraw
+{
+	GENERATED_BODY()
+
+#if WITH_EDITORONLY_DATA
+
+public:
+
+	/** Widget Radius */
+	UPROPERTY(EditAnywhere, Instanced, Category = "Settings")
+	TArray<TObjectPtr<UAnimDatabaseDebugDraw>> Drawers;
+
+#endif
+
+#if WITH_EDITOR
+
+public:
+
+	virtual void InitializeDrawDebug_Implementation(const UAnimDatabase* InDatabase, const FAnimDatabaseFrameRanges& InFrameRanges) override;
+
+	virtual void DrawDebug_Implementation(
+		const FDebugDrawer& Drawer,
+		const FAnimDatabasePoseState& InPoseState,
+		const UAnimDatabase* InDatabase,
+		const FAnimDatabaseFrameRanges& InFrameRanges,
+		const FTransform& RangeViewportTransform,
+		const int32 SequenceIdx,
+		const float SequenceTime,
+		const int32 RangeStart,
+		const int32 RangeLength,
+		const FLinearColor& IdentifierColor) override;
+
+#endif
+};
+
+/** Debug Draw class for drawing root orientation */
+UCLASS(BlueprintType, Blueprintable, meta = (DisplayName = "Root Orientation"))
+class UAnimDatabaseDebugDraw_RootOrientation : public UAnimDatabaseDebugDraw
+{
+	GENERATED_BODY()
+
+#if WITH_EDITORONLY_DATA
+
+public:
+
+	/** Vertical Offset from Floor */
+	UPROPERTY(EditAnywhere, Category = "Settings", meta = (ForceForceUnits = "cm"))
+	float FloorOffset = 1.0f;
+
+	/** Widget Radius */
+	UPROPERTY(EditAnywhere, Category = "Settings", meta = (ForceForceUnits = "cm"))
+	float Radius = 50.0f;
+
+	/** Local root forward direction */
+	UPROPERTY(EditAnywhere, Category = "Settings")
+	FVector ForwardVector = FVector::RightVector;
+
+#endif
+
+#if WITH_EDITOR
+
+public:
+
+	virtual void DrawDebug_Implementation(
+		const FDebugDrawer& Drawer,
+		const FAnimDatabasePoseState& InPoseState,
+		const UAnimDatabase* InDatabase,
+		const FAnimDatabaseFrameRanges& InFrameRanges,
+		const FTransform& RangeViewportTransform,
+		const int32 SequenceIdx,
+		const float SequenceTime,
+		const int32 RangeStart,
+		const int32 RangeLength,
+		const FLinearColor& IdentifierColor) override;
+
+#endif
+};
+
 /** Debug Draw class for drawing root velocities */
 UCLASS(BlueprintType, Blueprintable, meta = (DisplayName = "Root Velocity"))
 class UAnimDatabaseDebugDraw_RootVelocity : public UAnimDatabaseDebugDraw
diff --git a/Engine/Plugins/Experimental/Animation/AnimDatabase/Source/AnimDatabase/Public/AnimDatabaseFrameAttribute.h b/Engine/Plugins/Experimental/Animation/AnimDatabase/Source/AnimDatabase/Public/AnimDatabaseFrameAttribute.h
--- a/Engine/Plugins/Experimental/Animation/AnimDatabase/Source/AnimDatabase/Public/AnimDatabaseFrameAttribute.h
+++ b/Engine/Plugins/Experimental/Animation/AnimDatabase/Source/AnimDatabase/Public/AnimDatabaseFrameAttribute.h
@@ -1643,11 +1643,11 @@ public:
 public:
 
 	/** Converts a FrameAttribute object to a string. */
-	UFUNCTION(BlueprintPure, Category = "AnimDatabase", meta = (BlueprintAutocast, CompactNodeTitle = "->", AutoCreateRefTerm = "FrameAttribute"))
+	UFUNCTION(BlueprintPure, Category = "AnimDatabase", meta = (BlueprintAutocast, CompactNodeTitle = "->", AutoCreateRefTerm = "FrameAttribute", Keywords="Print,Log,Display,Name"))
 	static UE_API FString FrameAttributeToString(const FAnimDatabaseFrameAttribute& FrameAttribute);
 
 	/** Converts a FrameAttribute object to a string object in full. */
-	UFUNCTION(BlueprintPure, Category = "AnimDatabase")
+	UFUNCTION(BlueprintPure, Category = "AnimDatabase", meta= (Keywords = "Print,Log,Display,Name"))
 	static UE_API FString FrameAttributeToStringFormat(const UAnimDatabase* Database, const FAnimDatabaseFrameAttribute& FrameAttribute, const int32 Cutoff = 11);
 
 
diff --git a/Engine/Plugins/Experimental/Animation/AnimDatabase/Source/AnimDatabase/Public/AnimDatabaseFrameRanges.h b/Engine/Plugins/Experimental/Animation/AnimDatabase/Source/AnimDatabase/Public/AnimDatabaseFrameRanges.h
--- a/Engine/Plugins/Experimental/Animation/AnimDatabase/Source/AnimDatabase/Public/AnimDatabaseFrameRanges.h
+++ b/Engine/Plugins/Experimental/Animation/AnimDatabase/Source/AnimDatabase/Public/AnimDatabaseFrameRanges.h
@@ -195,11 +195,11 @@ public:
 	static UE_API FAnimDatabaseFrames FramesDifferenceFromArrayView(const TArrayView<const FAnimDatabaseFrames> Frames);
 
 	/** Converts a Frames object to a string. */
-	UFUNCTION(BlueprintPure, Category = "AnimDatabase", meta = (BlueprintAutocast, CompactNodeTitle = "->", AutoCreateRefTerm = "Frames"))
+	UFUNCTION(BlueprintPure, Category = "AnimDatabase", meta = (BlueprintAutocast, CompactNodeTitle = "->", AutoCreateRefTerm = "Frames", Keywords = "Print,Log,Display,Name"))
 	static UE_API FString FramesToString(const FAnimDatabaseFrames& Frames);
 
 	/** Converts a Frames object to a string object in full. */
-	UFUNCTION(BlueprintPure, Category = "AnimDatabase")
+	UFUNCTION(BlueprintPure, Category = "AnimDatabase", meta = (Keywords = "Print,Log,Display,Name"))
 	static UE_API FString FramesToStringFormat(const UAnimDatabase* Database, const FAnimDatabaseFrames& Frames, const int32 Cutoff = 11);
 };
 
@@ -617,11 +617,11 @@ public:
 public:
 
 	/** Converts a FrameRanges object to a string. */
-	UFUNCTION(BlueprintPure, Category = "AnimDatabase", meta = (BlueprintAutocast, CompactNodeTitle = "->", AutoCreateRefTerm = "FrameRanges"))
+	UFUNCTION(BlueprintPure, Category = "AnimDatabase", meta = (BlueprintAutocast, CompactNodeTitle = "->", AutoCreateRefTerm = "FrameRanges", Keywords = "Print,Log,Display,Name"))
 	static UE_API FString FrameRangesToString(const FAnimDatabaseFrameRanges& FrameRanges);
 
 	/** Converts a FrameRanges object to a string object in full. */
-	UFUNCTION(BlueprintPure, Category = "AnimDatabase")
+	UFUNCTION(BlueprintPure, Category = "AnimDatabase", meta=(Keywords = "Print,Log,Display,Name"))
 	static UE_API FString FrameRangesToStringFormat(const UAnimDatabase* Database, const FAnimDatabaseFrameRanges& FrameRanges, const int32 Cutoff = 11);
 
 public:
diff --git a/Engine/Plugins/Experimental/Animation/AnimDatabase/Source/AnimDatabase/Public/AnimDatabasePose.h b/Engine/Plugins/Experimental/Animation/AnimDatabase/Source/AnimDatabase/Public/AnimDatabasePose.h
--- a/Engine/Plugins/Experimental/Animation/AnimDatabase/Source/AnimDatabase/Public/AnimDatabasePose.h
+++ b/Engine/Plugins/Experimental/Animation/AnimDatabase/Source/AnimDatabase/Public/AnimDatabasePose.h
@@ -51,6 +51,7 @@ namespace UE::AnimDatabase
 	{
 		UE_API int32 GetFrameNum() const;
 
+		UE_API FPoseRootDataConstView ConstView() const;
 		UE_API FPoseRootDataView Slice(const int32 FrameStart, const int32 FrameNum) const;
 
 		TLearningArrayView<1, FVector> RootLocations;
@@ -895,8 +896,24 @@ namespace UE::AnimDatabase
 				const TArrayView<const int32> AnkleBoneIndices,
 				const TArrayView<const int32> ToeBoneIndices,
 				const TArrayView<const FVector3f> KneeSideVectors,
+				const TArrayView<const FVector3f> ToeForwardVectors,
 				const float DeltaTime,
 				const FContactPinningSettings& Settings = FContactPinningSettings());
+
+			/** Removes Foot-Ground penetration without pinning the feet */
+			UE_API void RemoveFootGroundPenetration(
+				const FPoseGlobalBoneDataView& InOutPoseGlobalBoneData,
+				const FPoseLocalBoneDataView& InOutPoseLocalBoneData,
+				const FPoseRootDataConstView& InPoseRootData,
+				const TArrayView<const int32> InBoneParents,
+				const TArrayView<const int32> PelvisBoneIndices,
+				const TArrayView<const int32> HipBoneIndices,
+				const TArrayView<const int32> KneeBoneIndices,
+				const TArrayView<const int32> AnkleBoneIndices,
+				const TArrayView<const int32> ToeBoneIndices,
+				const TArrayView<const FVector3f> KneeSideVectors,
+				const TArrayView<const FVector3f> ToeForwardVectors,
+				const FContactPinningSettings& Settings = FContactPinningSettings());
 		}
 	}
 
diff --git a/Engine/Plugins/Experimental/Animation/AnimDatabase/Source/AnimDatabaseEditor/Private/AnimDatabaseEditor.cpp b/Engine/Plugins/Experimental/Animation/AnimDatabase/Source/AnimDatabaseEditor/Private/AnimDatabaseEditor.cpp
--- a/Engine/Plugins/Experimental/Animation/AnimDatabase/Source/AnimDatabaseEditor/Private/AnimDatabaseEditor.cpp
+++ b/Engine/Plugins/Experimental/Animation/AnimDatabase/Source/AnimDatabaseEditor/Private/AnimDatabaseEditor.cpp
@@ -12,11 +12,7 @@
 void FAnimDatabaseEditorModule::StartupModule()
 {
 	// Register Ed Mode used to interact with the preview scene
-	FEditorModeRegistry::Get().RegisterMode<UE::AnimDatabase::Editor::FDatabaseMode>(UE::AnimDatabase::Editor::FDatabaseMode::EditorModeId,
-		LOCTEXT("AnimDatabaseEditorModeName", "AnimDatabase"),
-		FSlateIcon(),
-		false,
-		9000);
+	FEditorModeRegistry::Get().RegisterMode<UE::AnimDatabase::Editor::FDatabaseMode>(UE::AnimDatabase::Editor::FDatabaseMode::EditorModeId, LOCTEXT("AnimDatabaseEditorModeName", "AnimDatabase"));
 
 	FPropertyEditorModule& PropertyModule = FModuleManager::LoadModuleChecked<FPropertyEditorModule>("PropertyEditor");
 	PropertyModule.RegisterCustomClassLayout("AnimDatabaseQuery", FOnGetDetailCustomizationInstance::CreateStatic(&UE::AnimDatabase::Editor::FQueryDetails::MakeInstance));
diff --git a/Engine/Plugins/Experimental/Animation/AnimDatabase/Source/AnimDatabaseEditor/Private/AnimDatabaseEditorMode.cpp b/Engine/Plugins/Experimental/Animation/AnimDatabase/Source/AnimDatabaseEditor/Private/AnimDatabaseEditorMode.cpp
--- a/Engine/Plugins/Experimental/Animation/AnimDatabase/Source/AnimDatabaseEditor/Private/AnimDatabaseEditorMode.cpp
+++ b/Engine/Plugins/Experimental/Animation/AnimDatabase/Source/AnimDatabaseEditor/Private/AnimDatabaseEditorMode.cpp
@@ -177,7 +177,10 @@ namespace UE::AnimDatabase::Editor
 			RangeTrajectoryLocations.SetNum(QueryRangeNum);
 			RangeTrajectoryRotations.SetNum(QueryRangeNum);
 
-			Database->WaitForCompressionOnAnimSequencesFromArrayView(Database->Query->QueryRanges.FrameRangeSet->GetEntrySequences());
+			if (Database->Query->QueryRanges.IsValid())
+			{
+				Database->WaitForCompressionOnAnimSequencesFromArrayView(Database->Query->QueryRanges.FrameRangeSet->GetEntrySequences());
+			}
 
 			ParallelFor(QueryRangeNum, [this, Database](int32 RangeIdx)
 				{
diff --git a/Engine/Plugins/Experimental/Animation/AnimDatabase/Source/AnimDatabaseEditor/Private/AnimDatabaseEditorToolkit.cpp b/Engine/Plugins/Experimental/Animation/AnimDatabase/Source/AnimDatabaseEditor/Private/AnimDatabaseEditorToolkit.cpp
--- a/Engine/Plugins/Experimental/Animation/AnimDatabase/Source/AnimDatabaseEditor/Private/AnimDatabaseEditorToolkit.cpp
+++ b/Engine/Plugins/Experimental/Animation/AnimDatabase/Source/AnimDatabaseEditor/Private/AnimDatabaseEditorToolkit.cpp
@@ -56,7 +56,7 @@ TSoftClassPtr<UObject> UAnimDatabaseEditorAssetDefinition::GetAssetClass() const
 
 TConstArrayView<FAssetCategoryPath> UAnimDatabaseEditorAssetDefinition::GetAssetCategories() const
 {
-	static const auto Categories = { EAssetCategoryPaths::Animation };
+	static const auto Categories = { FAssetCategoryPath(EAssetCategoryPaths::Animation) };
 	return Categories;
 }
 
diff --git a/Engine/Plugins/Experimental/Animation/AnimGen/Source/AnimGen/Private/AnimGen.cpp b/Engine/Plugins/Experimental/Animation/AnimGen/Source/AnimGen/Private/AnimGen.cpp
new file mode 100644
--- /dev/null
+++ b/Engine/Plugins/Experimental/Animation/AnimGen/Source/AnimGen/Private/AnimGen.cpp
@@ -0,0 +1,30 @@
+// Copyright Epic Games, Inc. All Rights Reserved.
+
+#include "AnimGen.h"
+
+#include "RigVMCore/RigVMRegistry.h"
+
+#include "AnimGenController.h"
+#include "AnimGenControl.h"
+#include "RigVMFunction_AnimGen.h"
+
+void FAnimGenModule::StartupModule()
+{
+	static UScriptStruct* const AllowedStructTypes[] =
+	{
+		FAnimGenControlObject::StaticStruct(),
+		FAnimGenControlObjectElement::StaticStruct(),
+		FAnimGenNamedControlElement::StaticStruct(),
+	};
+
+	static TPair<UClass*, FRigVMRegistry::ERegisterObjectOperation> const AllowedObjectTypes[] =
+	{
+		{ UAnimGenController::StaticClass(), FRigVMRegistry::ERegisterObjectOperation::Class },
+	};
+
+	FRigVMRegistry& RigVMRegistry = FRigVMRegistry::Get();
+	RigVMRegistry.RegisterStructTypes(AllowedStructTypes);
+	RigVMRegistry.RegisterObjectTypes(AllowedObjectTypes);
+}
+	
+IMPLEMENT_MODULE(FAnimGenModule, AnimGen)
\ No newline at end of file
