AnimDatabase: Added function for stripping unused bone tracks and made python lazy-load.

#Synced-CL 51423447

diff --git a/Engine/Plugins/Experimental/Animation/AnimDatabase/Source/AnimDatabase/Private/AnimDatabase.cpp b/Engine/Plugins/Experimental/Animation/AnimDatabase/Source/AnimDatabase/Private/AnimDatabase.cpp
--- a/Engine/Plugins/Experimental/Animation/AnimDatabase/Source/AnimDatabase/Private/AnimDatabase.cpp
+++ b/Engine/Plugins/Experimental/Animation/AnimDatabase/Source/AnimDatabase/Private/AnimDatabase.cpp
@@ -6358,6 +6358,74 @@ void UAnimDatabaseLibrary::DatabaseSetLocalBoneTransformsFromArrayViews(UAnimDat
 	UE::AnimDatabase::Private::CloseControllers(Controllers, bShouldTransact);
 }
 
+void UAnimDatabaseLibrary::DatabaseStripUnusedBoneTracks(UAnimDatabase* Database, const FAnimDatabaseFrameRanges& FrameRanges, const bool bShouldTransact)
+{
+	const FAnimDatabaseFrameRanges FilteredFrameRanges = UAnimDatabaseFrameRangesLibrary::MakeFrameRangesFromNotMirrored(Database, FrameRanges);
+
+	// Sample pose data
+
+	const int32 TotalFrameNum = FilteredFrameRanges.FrameRangeSet->GetTotalFrameNum();
+	const int32 BoneNum = Database->GetBoneNum();
+
+	UE::AnimDatabase::FPoseData PoseData;
+	PoseData.Resize(TotalFrameNum, BoneNum);
+
+	Database->WaitForCompressionOnAnimSequencesFromArrayView(FilteredFrameRanges.FrameRangeSet->GetEntrySequences());
+
+	UE::Learning::FrameRangeSet::ParallelForEachRange(*FilteredFrameRanges.FrameRangeSet, [&FilteredFrameRanges, &PoseData, &Database](
+		const int32 TotalRangeIdx,
+		const int32 EntryIdx,
+		const int32 RangeIdx) {
+
+			TRACE_CPUPROFILER_EVENT_SCOPE(UAnimDatabaseLibrary::DatabaseModifyPoseData::SampleAndModify);
+
+			const int32 FrameOffset = FilteredFrameRanges.FrameRangeSet->GetEntryRangeOffset(EntryIdx, RangeIdx);
+			const int32 FrameNum = FilteredFrameRanges.FrameRangeSet->GetEntryRangeLength(EntryIdx, RangeIdx);
+			const int32 SequenceIdx = FilteredFrameRanges.FrameRangeSet->GetEntrySequence(EntryIdx);
+			const int32 StartFrame = FilteredFrameRanges.FrameRangeSet->GetEntryRangeStart(EntryIdx, RangeIdx);
+
+			Database->GetPoseData(PoseData.Slice(FrameOffset, FrameNum), SequenceIdx, StartFrame);
+		});
+
+	const FReferenceSkeleton ReferenceSkeleton = Database->Skeleton->GetReferenceSkeleton();
+	check(ReferenceSkeleton.GetRefBonePose().Num() == BoneNum);
+
+	TArray<int32> UsedBones;
+	UE::AnimDatabase::PoseData::ComputeUsedBones(
+		UsedBones,
+		PoseData.LocalBoneData.ConstView(),
+		ReferenceSkeleton.GetRefBonePose());
+
+	const int32 EntryNum = FilteredFrameRanges.FrameRangeSet->GetEntryNum();
+
+	for (int32 EntryIdx = 0; EntryIdx < EntryNum; EntryIdx++)
+	{
+		const int32 SequenceIdx = FilteredFrameRanges.FrameRangeSet->GetEntrySequence(EntryIdx);
+
+		if (UAnimSequence* AnimSequence = Database->Entries[SequenceIdx])
+		{
+			TScriptInterface<IAnimationDataController> Controller = AnimSequence->GetDataModel()->GetController();
+			Controller->OpenBracket(LOCTEXT("StrippingKeysScript", "Stripping Unused Keys"), bShouldTransact);
+
+			for (int32 BoneIdx = 0; BoneIdx < BoneNum; BoneIdx++)
+			{
+				if (!UsedBones.Contains(BoneIdx))
+				{
+					const FName BoneName = ReferenceSkeleton.GetBoneName(BoneIdx);
+					check(BoneName != NAME_None);
+
+					if (Controller->GetModel()->IsValidBoneTrackName(BoneName))
+					{
+						Controller->RemoveBoneTrack(BoneName, bShouldTransact);
+					}
+				}
+			}
+
+			Controller->CloseBracket();
+		}
+	}
+}
+
 void UAnimDatabaseLibrary::DatabaseModifyPoseData(UAnimDatabase* Database, const FAnimDatabaseFrameRanges& FrameRanges, const TFunctionRef<void(const UE::AnimDatabase::FPoseDataView& InOutPoseData)> Function, const bool bShouldTransact)
 {
 	const FAnimDatabaseFrameRanges FilteredFrameRanges = UAnimDatabaseFrameRangesLibrary::MakeFrameRangesFromNotMirrored(Database, FrameRanges);
@@ -6398,7 +6466,7 @@ void UAnimDatabaseLibrary::DatabaseModifyPoseData(UAnimDatabase* Database, const
 
 	UE::AnimDatabase::Private::OpenControllers(
 		Controllers, SequenceControllerIndices, Database, FilteredFrameRanges,
-		LOCTEXT("MakeLoopedScript", "Making Ranges Looped"), bShouldTransact);
+		LOCTEXT("ModifyPoseDataScript", "Modifying Pose Data"), bShouldTransact);
 
 	TLearningArray<1, int32> TagRangeSequenceIndices;
 	TagRangeSequenceIndices.SetNumUninitialized({ TotalRangeNum });
@@ -6417,12 +6485,18 @@ void UAnimDatabaseLibrary::DatabaseModifyPoseData(UAnimDatabase* Database, const
 		const int32 FrameNum = FilteredFrameRanges.FrameRangeSet->GetAllRangeLengths()[RangeIdx];
 		const int32 StartFrame = FilteredFrameRanges.FrameRangeSet->GetAllRangeStarts()[RangeIdx];
 
+		TArray<int32> UsedBones;
+		UE::AnimDatabase::PoseData::ComputeUsedBones(
+			UsedBones, 
+			PoseData.LocalBoneData.ConstSlice(FrameOffset, FrameNum), 
+			Database->Skeleton->GetReferenceSkeleton().GetRefBonePose());
+
 		SetAnimSequencePoseData(
 			Database->GetAnimSequence(RangeSequenceIdx),
 			Controller,
 			StartFrame,
 			PoseData.ConstSlice(FrameOffset, FrameNum),
-			UE::Learning::FIndexSet(0, BoneNum),
+			UsedBones,
 			bShouldTransact);
 	}
 
@@ -6563,6 +6637,11 @@ void UAnimDatabaseLibrary::DatabaseRemoveFootGroundPenetration(UAnimDatabase* Da
 		}, bShouldTransact);
 }
 
+void UAnimDatabaseFunction_StripUnusedBoneTracks::Run_Implementation(UAnimDatabase* InDatabase, const FAnimDatabaseFrameRanges& InFrameRanges, const bool bShouldTransact)
+{
+	UAnimDatabaseLibrary::DatabaseStripUnusedBoneTracks(InDatabase, InFrameRanges, bShouldTransact);
+}
+
 void UAnimDatabaseFunction_ExportBVH::Run_Implementation(UAnimDatabase* InDatabase, const FAnimDatabaseFrameRanges& InFrameRanges, const bool bShouldTransact)
 {
 	if (!ExportDirectory.Path.IsEmpty())
diff --git a/Engine/Plugins/Experimental/Animation/AnimDatabase/Source/AnimDatabase/Private/AnimDatabaseFrameRanges.cpp b/Engine/Plugins/Experimental/Animation/AnimDatabase/Source/AnimDatabase/Private/AnimDatabaseFrameRanges.cpp
--- a/Engine/Plugins/Experimental/Animation/AnimDatabase/Source/AnimDatabase/Private/AnimDatabaseFrameRanges.cpp
+++ b/Engine/Plugins/Experimental/Animation/AnimDatabase/Source/AnimDatabase/Private/AnimDatabaseFrameRanges.cpp
@@ -13,6 +13,11 @@
 #include "Animation/AnimNotifies/AnimNotify.h"
 #include "Animation/AnimNotifies/AnimNotifyState.h"
 
+#if WITH_EDITOR
+#include "Misc/PackageName.h"
+#include "HAL/FileManager.h"
+#endif
+
 #include "LearningRandom.h"
 
 #include UE_INLINE_GENERATED_CPP_BY_NAME(AnimDatabaseFrameRanges)
@@ -2062,7 +2067,11 @@ void UAnimDatabaseFrameRangesLibrary::FrameRangesStatistics(FAnimDatabaseFrameRa
 	OutStatistics.MaximumRangeDuration = FrameRanges.FrameRangeSet->GetTotalRangeNum() > 0 ? -UE_MAX_FLT : 0.0f;
 
 	const int32 EntryNum = FrameRanges.FrameRangeSet->GetEntryNum();
-	const float DeltaTime = 1.0f / Database->GetFrameRate().AsDecimal();
+	const float DeltaTime = 1.0f / FMath::Max(Database->GetFrameRate().AsDecimal(), UE_SMALL_NUMBER);
+
+#if WITH_EDITOR
+	IFileManager& FileManager = IFileManager::Get();
+#endif
 
 	TSet<const UAnimSequence*> ProcessedSequences;
 	ProcessedSequences.Reserve(EntryNum);
@@ -2075,13 +2084,20 @@ void UAnimDatabaseFrameRangesLibrary::FrameRangesStatistics(FAnimDatabaseFrameRa
 		const UAnimSequence* Sequence = Database->GetAnimSequence(SequenceIdx);
 		const float TotalSequenceDuration = Database->GetSequenceDuration(SequenceIdx);
 
-		int32 UncompressedSize = 0, CompressedSize = 0;
+		int32 UncompressedSize = 0, CompressedSize = 0, DiskSize = 0;
 		if (Sequence && !ProcessedSequences.Contains(Sequence))
 		{
 #if WITH_EDITOR
+			FString PackageFilename;
+			if (FPackageName::DoesPackageExist(Sequence->GetOutermost() ? Sequence->GetOutermost()->GetName() : TEXT(""), &PackageFilename))
+			{
+				DiskSize = FileManager.FileSize(*PackageFilename) / 1000;
+			}
+
 			UncompressedSize = Sequence->GetUncompressedRawSize() / 1000;
 #endif
 			CompressedSize = Sequence->GetApproxCompressedSize() / 1000;
+				
 			ProcessedSequences.Add(Sequence);
 		}
 
@@ -2095,8 +2111,9 @@ void UAnimDatabaseFrameRangesLibrary::FrameRangesStatistics(FAnimDatabaseFrameRa
 			OutStatistics.MinimumRangeDuration = FMath::Min(OutStatistics.MinimumRangeDuration, RangeDuration);
 			OutStatistics.MaximumRangeDuration = FMath::Max(OutStatistics.MaximumRangeDuration, RangeDuration);
 
-			OutStatistics.UncompressedSize += UncompressedSize * (RangeDuration / TotalSequenceDuration);
-			OutStatistics.CompressedSize += CompressedSize * (RangeDuration / TotalSequenceDuration);
+			OutStatistics.DiskSize += DiskSize * (RangeDuration / FMath::Max(TotalSequenceDuration, UE_SMALL_NUMBER));
+			OutStatistics.UncompressedSize += UncompressedSize * (RangeDuration / FMath::Max(TotalSequenceDuration, UE_SMALL_NUMBER));
+			OutStatistics.CompressedSize += CompressedSize * (RangeDuration / FMath::Max(TotalSequenceDuration, UE_SMALL_NUMBER));
 			OutStatistics.TotalFrameNum += RangeFrameNum;
 			OutStatistics.TotalRangeNum++;
 		}
diff --git a/Engine/Plugins/Experimental/Animation/AnimDatabase/Source/AnimDatabase/Public/AnimDatabase.h b/Engine/Plugins/Experimental/Animation/AnimDatabase/Source/AnimDatabase/Public/AnimDatabase.h
--- a/Engine/Plugins/Experimental/Animation/AnimDatabase/Source/AnimDatabase/Public/AnimDatabase.h
+++ b/Engine/Plugins/Experimental/Animation/AnimDatabase/Source/AnimDatabase/Public/AnimDatabase.h
@@ -999,6 +999,10 @@ public:
 	static UE_API void DatabaseSetLocalBoneTransforms(UAnimDatabase* Database, const TArray<FName>& BoneNames, const TArray<FAnimDatabaseFrameAttribute>& TransformFrameAttributes, const bool bShouldTransact = false);
 	static UE_API void DatabaseSetLocalBoneTransformsFromArrayViews(UAnimDatabase* Database, const TArrayView<const FName> BoneNames, const TArrayView<const FAnimDatabaseFrameAttribute> TransformFrameAttributes, const bool bShouldTransact = false);
 
+	/** Strips unused bone tracks from the data to reduce the file size on disk */
+	UFUNCTION(BlueprintCallable, Category = "AnimDatabase", meta = (DevelopmentOnly))
+	static UE_API void DatabaseStripUnusedBoneTracks(UAnimDatabase* Database, const FAnimDatabaseFrameRanges& FrameRanges, const bool bShouldTransact = false);
+
 	/** Uses the provided function to modify the database pose data */
 	static UE_API void DatabaseModifyPoseData(UAnimDatabase* Database, const FAnimDatabaseFrameRanges& FrameRanges, const TFunctionRef<void(const UE::AnimDatabase::FPoseDataView& InOutPoseData)> Function, const bool bShouldTransact = false);
 
@@ -1026,6 +1030,22 @@ private:
 #endif
 };
 
+/** Function for stripping unused bone-tracks from the AnimSequence to reduce the file size on desk  */
+UCLASS(BlueprintType, Blueprintable, meta = (DisplayName = "Strip Unused Bone Tracks"))
+class UAnimDatabaseFunction_StripUnusedBoneTracks : public UAnimDatabaseFunction
+{
+	GENERATED_BODY()
+
+public:
+
+#if WITH_EDITOR
+
+public:
+
+	UE_API virtual void Run_Implementation(UAnimDatabase* InDatabase, const FAnimDatabaseFrameRanges& InFrameRanges, const bool bShouldTransact) override;
+
+#endif
+};
 
 /** Function for exporting the given ranges as BVH  */
 UCLASS(BlueprintType, Blueprintable, meta = (DisplayName = "Export to BVH"))
diff --git a/Engine/Plugins/Experimental/Animation/AnimDatabase/Source/AnimDatabase/Public/AnimDatabaseFrameRanges.h b/Engine/Plugins/Experimental/Animation/AnimDatabase/Source/AnimDatabase/Public/AnimDatabaseFrameRanges.h
--- a/Engine/Plugins/Experimental/Animation/AnimDatabase/Source/AnimDatabase/Public/AnimDatabaseFrameRanges.h
+++ b/Engine/Plugins/Experimental/Animation/AnimDatabase/Source/AnimDatabase/Public/AnimDatabaseFrameRanges.h
@@ -268,6 +268,10 @@ struct FAnimDatabaseFrameRangesStatistics
 {
 	GENERATED_BODY()
 
+	/** Approximate size of the animation data in the ranges on disk in kilobytes. Will only count mirrored data once. */
+	UPROPERTY(VisibleAnywhere, Category = "Statistics", meta = (Units = "Kilobytes"))
+	float DiskSize = 0.0f;
+
 	/** Approximate size of the uncompressed animation data in the ranges in kilobytes. Will only count mirrored data once. */
 	UPROPERTY(VisibleAnywhere, Category = "Statistics", meta = (Units = "Kilobytes"))
 	float UncompressedSize = 0.0f;
diff --git a/Engine/Plugins/Experimental/Animation/AnimGen/Source/AnimGenEditor/AnimGenEditor.Build.cs b/Engine/Plugins/Experimental/Animation/AnimGen/Source/AnimGenEditor/AnimGenEditor.Build.cs
--- a/Engine/Plugins/Experimental/Animation/AnimGen/Source/AnimGenEditor/AnimGenEditor.Build.cs
+++ b/Engine/Plugins/Experimental/Animation/AnimGen/Source/AnimGenEditor/AnimGenEditor.Build.cs
@@ -70,6 +70,7 @@ public class AnimGenEditor : ModuleRules
 				"NNERuntimeBasicCpu",
 				"Blutility",
 				"DrawDebugLibrary",
+				"PythonScriptPlugin",
 			}
 			);
 		
diff --git a/Engine/Plugins/Experimental/Animation/AnimGen/Source/AnimGenEditor/Private/AnimGenEditorAutoEncoderToolkit.cpp b/Engine/Plugins/Experimental/Animation/AnimGen/Source/AnimGenEditor/Private/AnimGenEditorAutoEncoderToolkit.cpp
--- a/Engine/Plugins/Experimental/Animation/AnimGen/Source/AnimGenEditor/Private/AnimGenEditorAutoEncoderToolkit.cpp
+++ b/Engine/Plugins/Experimental/Animation/AnimGen/Source/AnimGenEditor/Private/AnimGenEditorAutoEncoderToolkit.cpp
@@ -51,6 +51,7 @@
 #include "Editor.h"
 #include "Subsystems/AssetEditorSubsystem.h"
 #include "AssetTypeCategories.h"
+#include "PipInstallHelpers.h"
 
 #include "Animation/Skeleton.h"
 #include "Animation/AnimSequence.h"
@@ -93,16 +94,26 @@ TConstArrayView<FAssetCategoryPath> UAnimGenEditorAutoEncoderAssetDefinition::Ge
 
 EAssetCommandResult UAnimGenEditorAutoEncoderAssetDefinition::OpenAssets(const FAssetOpenArgs& OpenArgs) const
 {
-	for (UAnimGenAutoEncoder* Asset : OpenArgs.LoadObjects<UAnimGenAutoEncoder>())
-	{
-		TSharedRef<UE::AnimGen::Editor::FAutoEncoderToolkit> NewEditor(new UE::AnimGen::Editor::FAutoEncoderToolkit());
-		NewEditor->InitAssetEditor(OpenArgs.GetToolkitMode(), OpenArgs.ToolkitHost, Asset);
-	}
+	FPythonScriptInitHelper::InitPythonAndPipInstall(FSimpleDelegate::CreateLambda([
+		OpenArgs = FAssetOpenArgs(OpenArgs),
+		// Taking a copy of OpenArgs.Assets as it's a TArrayView and this callback outlives the view
+		OpenAssetsData = TArray<FAssetData>(OpenArgs.Assets)]()
+		{
+			FAssetArgs WrappedArgs(OpenAssetsData);
+			for (UAnimGenAutoEncoder* Asset : WrappedArgs.LoadObjects<UAnimGenAutoEncoder>())
+			{
+				TSharedRef<UE::AnimGen::Editor::FAutoEncoderToolkit> NewEditor(new UE::AnimGen::Editor::FAutoEncoderToolkit());
+				NewEditor->InitAssetEditor(OpenArgs.GetToolkitMode(), OpenArgs.ToolkitHost, Asset);
+			}
+		}),
+		FSimpleDelegate::CreateLambda([]()
+			{
+				UE_LOG(LogAnimGen, Warning, TEXT("AnimGen AutoEncoder toolkit may not function properly without Python enabled"));
+			}));
 
 	return EAssetCommandResult::Handled;
 }
 
-
 UAnimGenAutoEncoderEditorFactory::UAnimGenAutoEncoderEditorFactory()
 {
 	bCreateNew = true;
diff --git a/Engine/Plugins/Experimental/Animation/AnimGen/Source/AnimGenEditor/Private/AnimGenEditorControllerToolkit.cpp b/Engine/Plugins/Experimental/Animation/AnimGen/Source/AnimGenEditor/Private/AnimGenEditorControllerToolkit.cpp
--- a/Engine/Plugins/Experimental/Animation/AnimGen/Source/AnimGenEditor/Private/AnimGenEditorControllerToolkit.cpp
+++ b/Engine/Plugins/Experimental/Animation/AnimGen/Source/AnimGenEditor/Private/AnimGenEditorControllerToolkit.cpp
@@ -48,6 +48,7 @@
 #include "Widgets/Input/SHyperlink.h"
 #include "Widgets/Layout/SScrollBox.h"
 #include "AssetTypeCategories.h"
+#include "PipInstallHelpers.h"
 
 #include "Animation/Skeleton.h"
 #include "Animation/AnimSequence.h"
@@ -91,11 +92,22 @@ TConstArrayView<FAssetCategoryPath> UAnimGenEditorControllerAssetDefinition::Get
 
 EAssetCommandResult UAnimGenEditorControllerAssetDefinition::OpenAssets(const FAssetOpenArgs& OpenArgs) const
 {
-	for (UAnimGenController* Asset : OpenArgs.LoadObjects<UAnimGenController>())
-	{
-		TSharedRef<UE::AnimGen::Editor::FControllerToolkit> NewEditor(new UE::AnimGen::Editor::FControllerToolkit());
-		NewEditor->InitAssetEditor(OpenArgs.GetToolkitMode(), OpenArgs.ToolkitHost, Asset);
-	}
+	FPythonScriptInitHelper::InitPythonAndPipInstall(FSimpleDelegate::CreateLambda([
+		OpenArgs = FAssetOpenArgs(OpenArgs),
+		// Taking a copy of OpenArgs.Assets as it's a TArrayView and this callback outlives the view
+		OpenAssetsData = TArray<FAssetData>(OpenArgs.Assets)]()
+		{
+			FAssetArgs WrappedArgs(OpenAssetsData);
+			for (UAnimGenController* Asset : WrappedArgs.LoadObjects<UAnimGenController>())
+			{
+				TSharedRef<UE::AnimGen::Editor::FControllerToolkit> NewEditor(new UE::AnimGen::Editor::FControllerToolkit());
+				NewEditor->InitAssetEditor(OpenArgs.GetToolkitMode(), OpenArgs.ToolkitHost, Asset);
+			}
+		}),
+		FSimpleDelegate::CreateLambda([]()
+			{
+				UE_LOG(LogAnimGen, Warning, TEXT("AnimGen Controller toolkit may not function properly without Python enabled"));
+			}));
 
 	return EAssetCommandResult::Handled;
 }
