AnimDatabase: Various small updates to frame attribute functions.

#Synced-CL 51496719

diff --git a/Engine/Plugins/Experimental/Animation/AnimDatabase/Source/AnimDatabase/Private/AnimDatabase.cpp b/Engine/Plugins/Experimental/Animation/AnimDatabase/Source/AnimDatabase/Private/AnimDatabase.cpp
--- a/Engine/Plugins/Experimental/Animation/AnimDatabase/Source/AnimDatabase/Private/AnimDatabase.cpp
+++ b/Engine/Plugins/Experimental/Animation/AnimDatabase/Source/AnimDatabase/Private/AnimDatabase.cpp
@@ -4188,8 +4188,11 @@ void UAnimDatabase::SampleCurveData(
 
 void UAnimDatabaseFunction::Run_Implementation(UAnimDatabase* InDatabase, const FAnimDatabaseFrameRanges& InFrameRanges, const bool bShouldTransact) {}
 
+#endif
+
 void UAnimDatabaseLibrary::GetAnimSequenceAnimNotifyClasses(TArray<TSubclassOf<UAnimNotify>>& OutClasses, UAnimSequence* AnimSequence)
 {
+#if WITH_EDITOR
 	OutClasses.Reset();
 
 	if (!AnimSequence)
@@ -4208,10 +4211,12 @@ void UAnimDatabaseLibrary::GetAnimSequenceAnimNotifyClasses(TArray<TSubclassOf<U
 			OutClasses.AddUnique(AnimNotifyRef.GetNotify()->Notify->GetClass());
 		}
 	}
+#endif
 }
 
 void UAnimDatabaseLibrary::GetAnimSequenceAnimNotifyStateClasses(TArray<TSubclassOf<UAnimNotifyState>>& OutClasses, UAnimSequence* AnimSequence)
 {
+#if WITH_EDITOR
 	OutClasses.Reset();
 
 	if (!AnimSequence)
@@ -4230,10 +4235,12 @@ void UAnimDatabaseLibrary::GetAnimSequenceAnimNotifyStateClasses(TArray<TSubclas
 			OutClasses.AddUnique(AnimNotifyRef.GetNotify()->NotifyStateClass->GetClass());
 		}
 	}
+#endif
 }
 
 void UAnimDatabaseLibrary::GetAnimSequenceAnimNotifyTimes(TArray<float>& OutTimes, const TSubclassOf<UAnimNotify>& Class, UAnimSequence* AnimSequence)
 {
+#if WITH_EDITOR
 	OutTimes.Reset();
 
 	if (!AnimSequence)
@@ -4256,10 +4263,12 @@ void UAnimDatabaseLibrary::GetAnimSequenceAnimNotifyTimes(TArray<float>& OutTime
 			OutTimes.Add(AnimNotifyRef.GetNotify()->GetTime());
 		}
 	}
+#endif
 }
 
 void UAnimDatabaseLibrary::GetAnimSequenceAnimNotifyStateTimesAndDurations(TArray<float>& OutTimes, TArray<float>& OutDurations, const TSubclassOf<UAnimNotifyState>& Class, UAnimSequence* AnimSequence)
 {
+#if WITH_EDITOR
 	OutTimes.Reset();
 	OutDurations.Reset();
 
@@ -4284,6 +4293,7 @@ void UAnimDatabaseLibrary::GetAnimSequenceAnimNotifyStateTimesAndDurations(TArra
 			OutDurations.Add(AnimNotifyRef.GetNotify()->GetDuration());
 		}
 	}
+#endif
 }
 
 void UAnimDatabaseLibrary::SetAnimSequencePoseData(
@@ -4294,6 +4304,7 @@ void UAnimDatabaseLibrary::SetAnimSequencePoseData(
 	const UE::Learning::FIndexSet UsedBones,
 	const bool bShouldTransact)
 {
+#if WITH_EDITOR
 	Controller->OpenBracket(LOCTEXT("KeyAnimSequence", "Keying Animation Sequence"), bShouldTransact);
 
 	const int32 FrameNum = PoseData.GetFrameNum();
@@ -4344,6 +4355,7 @@ void UAnimDatabaseLibrary::SetAnimSequencePoseData(
 
 	Controller->NotifyPopulated();
 	Controller->CloseBracket(bShouldTransact);
+#endif
 }
 
 void UAnimDatabaseLibrary::AddAnimNotifiesToSequence(UAnimSequence* AnimSequence, const FName TrackName, const TSubclassOf<UAnimNotify>& Class, const TArray<float>& Times)
@@ -4353,6 +4365,7 @@ void UAnimDatabaseLibrary::AddAnimNotifiesToSequence(UAnimSequence* AnimSequence
 
 void UAnimDatabaseLibrary::AddAnimNotifiesToSequenceArrayView(UAnimSequence* AnimSequence, const FName TrackName, const TSubclassOf<UAnimNotify>& Class, const TArrayView<const float> Times)
 {
+#if WITH_EDITOR
 	if (!AnimSequence)
 	{
 		UE_LOGFMT(LogAnimDatabase, Error, "AddAnimNotifiesToSequenceArrayView: AnimSequence is null.");
@@ -4400,10 +4413,12 @@ void UAnimDatabaseLibrary::AddAnimNotifiesToSequenceArrayView(UAnimSequence* Ani
 	// Refresh all cached data
 	AnimSequence->RefreshCacheData();
 	AnimSequence->Modify(true);
+#endif
 }
 
 void UAnimDatabaseLibrary::RemoveAllAnimNotifiesAndAnimNotifyStatesFromSequence(UAnimSequence* AnimSequence)
 {
+#if WITH_EDITOR
 	if (!AnimSequence)
 	{
 		UE_LOGFMT(LogAnimDatabase, Error, "RemoveAllAnimNotifiesAndAnimNotifyStatesFromSequence: AnimSequence is null.");
@@ -4414,6 +4429,7 @@ void UAnimDatabaseLibrary::RemoveAllAnimNotifiesAndAnimNotifyStatesFromSequence(
 	AnimSequence->Notifies.Empty();
 	AnimSequence->RefreshCacheData();
 	AnimSequence->Modify(true);
+#endif
 }
 
 void UAnimDatabaseLibrary::AddAnimNotifyStatesToSequence(UAnimSequence* AnimSequence, const FName TrackName, const TSubclassOf<UAnimNotifyState>& Class, const TArray<float>& Times, const TArray<float>& Durations)
@@ -4423,6 +4439,7 @@ void UAnimDatabaseLibrary::AddAnimNotifyStatesToSequence(UAnimSequence* AnimSequ
 
 void UAnimDatabaseLibrary::AddAnimNotifyStatesToSequenceArrayView(UAnimSequence* AnimSequence, const FName TrackName, const TSubclassOf<UAnimNotifyState>& Class, const TArrayView<const float> Times, const TArrayView<const float> Durations)
 {
+#if WITH_EDITOR
 	if (!AnimSequence)
 	{
 		UE_LOGFMT(LogAnimDatabase, Error, "AddAnimNotifyStatesToSequenceArrayView: AnimSequence is null.");
@@ -4472,6 +4489,7 @@ void UAnimDatabaseLibrary::AddAnimNotifyStatesToSequenceArrayView(UAnimSequence*
 	// Refresh all cached data
 	AnimSequence->RefreshCacheData();
 	AnimSequence->Modify(true);
+#endif
 }
 
 void UAnimDatabaseLibrary::EnableRootMotionAndForceRootLockOnSequences(const TArray<UAnimSequence*>& Sequences)
@@ -4481,6 +4499,7 @@ void UAnimDatabaseLibrary::EnableRootMotionAndForceRootLockOnSequences(const TAr
 
 void UAnimDatabaseLibrary::EnableRootMotionAndForceRootLockOnSequencesArrayView(TConstArrayView<UAnimSequence*> Sequences)
 {
+#if WITH_EDITOR
 	for (UAnimSequence* Sequence : Sequences)
 	{
 		if (Sequence)
@@ -4490,10 +4509,12 @@ void UAnimDatabaseLibrary::EnableRootMotionAndForceRootLockOnSequencesArrayView(
 			Sequence->Modify(true);
 		}
 	}
+#endif
 }
 
 void UAnimDatabaseLibrary::DatabaseRemoveNotifies(UAnimDatabase* Database, const FAnimDatabaseFrameRanges& FrameRanges)
 {
+#if WITH_EDITOR
 	if (!Database)
 	{
 		UE_LOGFMT(LogAnimDatabase, Error, "DatabaseRemoveNotifies: Database is null.");
@@ -4572,10 +4593,12 @@ void UAnimDatabaseLibrary::DatabaseRemoveNotifies(UAnimDatabase* Database, const
 			}
 		}
 	}
+#endif
 }
 
 void UAnimDatabaseLibrary::DatabaseExportNotifies(UAnimDatabase* Database, const FAnimDatabaseFrameRanges& FrameRanges, const FFilePath& Path)
 {
+#if WITH_EDITOR
 	if (!Database)
 	{
 		UE_LOGFMT(LogAnimDatabase, Error, "DatabaseExportNotifies: Database is null.");
@@ -4722,10 +4745,12 @@ void UAnimDatabaseLibrary::DatabaseExportNotifies(UAnimDatabase* Database, const
 			UE_LOGFMT(LogAnimDatabase, Error, "DatabaseExportNotifies: Failed to save to file {Filename}.", *Path.FilePath);
 		}
 	}
+#endif
 }
 
 void UAnimDatabaseLibrary::LoadAllBlueprintAssetsOfClass(UClass* Class)
 {
+#if WITH_EDITOR
 	FAssetRegistryModule& AssetRegistryModule = FModuleManager::LoadModuleChecked<FAssetRegistryModule>(TEXT("AssetRegistry"));
 	IAssetRegistry& AssetRegistry = AssetRegistryModule.Get();
 
@@ -4749,10 +4774,12 @@ void UAnimDatabaseLibrary::LoadAllBlueprintAssetsOfClass(UClass* Class)
 			LoadObject<UClass>(nullptr, *FPackageName::ExportTextPathToObjectPath(GeneratedClassPath));
 		}
 	}
+#endif
 }
 
 void UAnimDatabaseLibrary::DatabaseImportNotifies(UAnimDatabase* Database, const FAnimDatabaseFrameRanges& FrameRanges, const FFilePath& Path)
 {
+#if WITH_EDITOR
 	if (!Database)
 	{
 		UE_LOGFMT(LogAnimDatabase, Error, "DatabaseImportNotifies: Database is null.");
@@ -4961,10 +4988,12 @@ void UAnimDatabaseLibrary::DatabaseImportNotifies(UAnimDatabase* Database, const
 	{
 		UE_LOGFMT(LogAnimDatabase, Error, "DatabaseImportNotifies: Failed to parse JSON.");
 	}
+#endif
 }
 
 void UAnimDatabaseLibrary::DatabaseExportAsBVH(UAnimDatabase* Database, const FAnimDatabaseFrameRanges& FrameRanges, const FDirectoryPath& Path, const bool bAddRootAxesChangeRotation, const EAnimDatabaseRotationOrder RotationOrder)
 {
+#if WITH_EDITOR
 	if (!FrameRanges.IsValid())
 	{
 		UE_LOGFMT(LogAnimDatabase, Error, "DatabaseExportAsBVH: Invalid Frame Ranges.");
@@ -5125,10 +5154,12 @@ void UAnimDatabaseLibrary::DatabaseExportAsBVH(UAnimDatabase* Database, const FA
 				}
 			}
 		});
+#endif
 }
 
 void UAnimDatabaseLibrary::DatabaseImportFromBVH(UAnimDatabase* Database, const FDirectoryPath& ImportPath, const FDirectoryPath& AssetPath, const bool bRemoveRootAxesChangeRotation, const bool bAddToDatabase, const bool bIgnoreMirrored)
 {
+#if WITH_EDITOR
 	if (!Database)
 	{
 		UE_LOGFMT(LogAnimDatabase, Error, "DatabaseImportFromBVH: Database is null.");
@@ -5552,10 +5583,12 @@ void UAnimDatabaseLibrary::DatabaseImportFromBVH(UAnimDatabase* Database, const
 			}
 		}
 	}
+#endif
 }
 
 void UAnimDatabaseLibrary::DatabaseExportAsAnimSequences(UAnimDatabase* Database, const FAnimDatabaseFrameRanges& FrameRanges, const FDirectoryPath& Path, const FString& AssetNameFormatString, const bool bExportMirrored, const bool bShouldTransact)
 {
+#if WITH_EDITOR
 	if (!FrameRanges.IsValid())
 	{
 		UE_LOGFMT(LogAnimDatabase, Error, "DatabaseExportAsAnimSequences: Invalid Frame Ranges.");
@@ -5712,10 +5745,12 @@ void UAnimDatabaseLibrary::DatabaseExportAsAnimSequences(UAnimDatabase* Database
 			}
 		}
 	}
+#endif
 }
 
 void UAnimDatabaseLibrary::DatabaseExportRootMotion(UAnimDatabase* Database, const FAnimDatabaseFrameRanges& FrameRanges, const FFilePath& Path)
 {
+#if WITH_EDITOR
 	if (!FrameRanges.IsValid())
 	{
 		UE_LOGFMT(LogAnimDatabase, Error, "DatabaseExportRootMotion: Invalid Frame Ranges.");
@@ -5821,6 +5856,7 @@ void UAnimDatabaseLibrary::DatabaseExportRootMotion(UAnimDatabase* Database, con
 			UE_LOGFMT(LogAnimDatabase, Error, "DatabaseExportRootMotion: Failed to save to file {Filename}.", *Path.FilePath);
 		}
 	}
+#endif
 }
 
 void UAnimDatabaseLibrary::DatabaseAddCurve(UAnimDatabase* Database, const FName CurveName, const FAnimDatabaseFrameAttribute& FloatFrameAttribute, const bool bSparseKeys, const bool bShouldTransact)
@@ -5835,6 +5871,7 @@ void UAnimDatabaseLibrary::DatabaseAddCurves(UAnimDatabase* Database, const TArr
 
 void UAnimDatabaseLibrary::DatabaseAddCurvesFromArrayViews(UAnimDatabase* Database, const TArrayView<const FName> CurveNames, const TArrayView<const FAnimDatabaseFrameAttribute> FloatFrameAttributes, const bool bSparseKeys, const bool bShouldTransact)
 {
+#if WITH_EDITOR
 	TRACE_CPUPROFILER_EVENT_SCOPE(UAnimDatabaseLibrary::DatabaseAddCurvesFromArrayViews);
 
 	if (!Database)
@@ -5923,6 +5960,7 @@ void UAnimDatabaseLibrary::DatabaseAddCurvesFromArrayViews(UAnimDatabase* Databa
 	}
 
 	UE::AnimDatabase::Private::CloseControllers(Controllers, bShouldTransact);
+#endif
 }
 
 void UAnimDatabaseLibrary::DatabaseAddAnimNotifyState(UAnimDatabase* Database, const FName AnimNotifyStateTrackName, const TSubclassOf<UAnimNotifyState> Class, const FAnimDatabaseFrameRanges& FrameRanges)
@@ -5937,6 +5975,7 @@ void UAnimDatabaseLibrary::DatabaseAddAnimNotifyStates(UAnimDatabase* Database,
 
 void UAnimDatabaseLibrary::DatabaseAddAnimNotifyStatesFromArrayViews(UAnimDatabase* Database, const TArrayView<const FName> AnimNotifyStateTrackNames, const TArrayView<const TSubclassOf<UAnimNotifyState>> Classes, const TArrayView<const FAnimDatabaseFrameRanges> FrameRanges)
 {
+#if WITH_EDITOR
 	TRACE_CPUPROFILER_EVENT_SCOPE(UAnimDatabaseLibrary::DatabaseAddAnimNotifyStatesFromArrayViews);
 
 	if (!Database)
@@ -6014,10 +6053,12 @@ void UAnimDatabaseLibrary::DatabaseAddAnimNotifyStatesFromArrayViews(UAnimDataba
 			}
 		}
 	}
+#endif
 }
 
 void UAnimDatabaseLibrary::DatabaseAddAnimNotifyObject(UAnimDatabase* Database, const FAnimDatabaseFrames& Frames, const FName& TrackName, const UAnimNotify* AnimNotify)
 {
+#if WITH_EDITOR
 	TRACE_CPUPROFILER_EVENT_SCOPE(UAnimDatabaseLibrary::DatabaseAddAnimNotifyObject);
 
 	if (!Database)
@@ -6077,10 +6118,12 @@ void UAnimDatabaseLibrary::DatabaseAddAnimNotifyObject(UAnimDatabase* Database,
 			AnimSequence->Modify(true);
 		}
 	}
+#endif
 }
 
 void UAnimDatabaseLibrary::DatabaseAddAnimNotifyStateObject(UAnimDatabase* Database, const FAnimDatabaseFrameRanges& FrameRanges, const FName& TrackName, const UAnimNotifyState* AnimNotifyState)
 {
+#if WITH_EDITOR
 	TRACE_CPUPROFILER_EVENT_SCOPE(UAnimDatabaseLibrary::DatabaseAddAnimNotifyStateObject);
 
 	if (!Database)
@@ -6141,6 +6184,7 @@ void UAnimDatabaseLibrary::DatabaseAddAnimNotifyStateObject(UAnimDatabase* Datab
 			AnimSequence->Modify(true);
 		}
 	}
+#endif
 }
 
 void UAnimDatabaseLibrary::DatabaseAddAnimNotify(UAnimDatabase* Database, const FName AnimNotifyTrackName, const TSubclassOf<UAnimNotify> Class, const FAnimDatabaseFrames& Frames)
@@ -6155,6 +6199,7 @@ void UAnimDatabaseLibrary::DatabaseAddAnimNotifies(UAnimDatabase* Database, cons
 
 void UAnimDatabaseLibrary::DatabaseAddAnimNotifiesFromArrayViews(UAnimDatabase* Database, const TArrayView<const FName> AnimNotifyTrackNames, const TArrayView<const TSubclassOf<UAnimNotify>> Classes, const TArrayView<const FAnimDatabaseFrames> Frames)
 {
+#if WITH_EDITOR
 	TRACE_CPUPROFILER_EVENT_SCOPE(UAnimDatabaseLibrary::DatabaseAddAnimNotifiesFromArrayViews);
 
 	if (AnimNotifyTrackNames.Num() != Classes.Num() ||
@@ -6229,6 +6274,7 @@ void UAnimDatabaseLibrary::DatabaseAddAnimNotifiesFromArrayViews(UAnimDatabase*
 			}
 		}
 	}
+#endif
 }
 
 void UAnimDatabaseLibrary::DatabaseSetLocalBoneTransform(UAnimDatabase* Database, const FName BoneName, const FAnimDatabaseFrameAttribute& TransformFrameAttribute, const bool bShouldTransact)
@@ -6243,6 +6289,7 @@ void UAnimDatabaseLibrary::DatabaseSetLocalBoneTransforms(UAnimDatabase* Databas
 
 void UAnimDatabaseLibrary::DatabaseSetLocalBoneTransformsFromArrayViews(UAnimDatabase* Database, const TArrayView<const FName> BoneNames, const TArrayView<const FAnimDatabaseFrameAttribute> TransformFrameAttributes, const bool bShouldTransact)
 {
+#if WITH_EDITOR
 	TRACE_CPUPROFILER_EVENT_SCOPE(UAnimDatabaseLibrary::DatabaseSetLocalBoneTransformsFromArrayViews);
 
 	if (BoneNames.Num() != TransformFrameAttributes.Num())
@@ -6356,10 +6403,12 @@ void UAnimDatabaseLibrary::DatabaseSetLocalBoneTransformsFromArrayViews(UAnimDat
 	}
 
 	UE::AnimDatabase::Private::CloseControllers(Controllers, bShouldTransact);
+#endif
 }
 
 void UAnimDatabaseLibrary::DatabaseStripUnusedBoneTracks(UAnimDatabase* Database, const FAnimDatabaseFrameRanges& FrameRanges, const bool bShouldTransact)
 {
+#if WITH_EDITOR
 	const FAnimDatabaseFrameRanges FilteredFrameRanges = UAnimDatabaseFrameRangesLibrary::MakeFrameRangesFromNotMirrored(Database, FrameRanges);
 
 	// Sample pose data
@@ -6424,10 +6473,12 @@ void UAnimDatabaseLibrary::DatabaseStripUnusedBoneTracks(UAnimDatabase* Database
 			Controller->CloseBracket();
 		}
 	}
+#endif
 }
 
 void UAnimDatabaseLibrary::DatabaseModifyPoseData(UAnimDatabase* Database, const FAnimDatabaseFrameRanges& FrameRanges, const TFunctionRef<void(const UE::AnimDatabase::FPoseDataView& InOutPoseData)> Function, const bool bShouldTransact)
 {
+#if WITH_EDITOR
 	const FAnimDatabaseFrameRanges FilteredFrameRanges = UAnimDatabaseFrameRangesLibrary::MakeFrameRangesFromNotMirrored(Database, FrameRanges);
 	
 	// Sample and modify pose data
@@ -6501,10 +6552,12 @@ void UAnimDatabaseLibrary::DatabaseModifyPoseData(UAnimDatabase* Database, const
 	}
 
 	UE::AnimDatabase::Private::CloseControllers(Controllers, bShouldTransact);
+#endif
 }
 
 void UAnimDatabaseLibrary::DatabaseMakeLooped(UAnimDatabase* Database, const FAnimDatabaseFrameRanges& FrameRanges, const float StartEndRatio, const float BlendInTime, const float BlendOutTime, const bool bShouldTransact)
 {
+#if WITH_EDITOR
 	TRACE_CPUPROFILER_EVENT_SCOPE(UAnimDatabaseLibrary::DatabaseMakeLooped);
 
 	DatabaseModifyPoseData(Database, FrameRanges, [Database, BlendInTime, BlendOutTime, StartEndRatio](const UE::AnimDatabase::FPoseDataView& InOutPoseData)
@@ -6512,10 +6565,12 @@ void UAnimDatabaseLibrary::DatabaseMakeLooped(UAnimDatabase* Database, const FAn
 			UE::AnimDatabase::PoseData::MakeLooped(InOutPoseData, Database->GetFrameRate(), StartEndRatio, BlendInTime, BlendOutTime);
 
 		}, bShouldTransact);
+#endif
 }
 
 void UAnimDatabaseLibrary::DatabasePatchRanges(UAnimDatabase* Database, const FAnimDatabaseFrameRanges& FrameRanges, const bool bApplyToRoot, const TArray<int32>& BoneIndices, const bool bShouldTransact)
 {
+#if WITH_EDITOR
 	TRACE_CPUPROFILER_EVENT_SCOPE(UAnimDatabaseLibrary::DatabasePatchRanges);
 
 	if (!Database) { return; }
@@ -6525,10 +6580,12 @@ void UAnimDatabaseLibrary::DatabasePatchRanges(UAnimDatabase* Database, const FA
 			UE::AnimDatabase::PoseData::PatchInterpolate(InOutPoseData, Database->GetFrameRate(), bApplyToRoot, BoneIndices);
 
 		}, bShouldTransact);
+#endif
 }
 
 void UAnimDatabaseLibrary::DatabaseSetToReferencePose(UAnimDatabase* Database, const FAnimDatabaseFrameRanges& FrameRanges, const bool bApplyToRoot, const TArray<int32>& IgnoreBoneIndices, const bool bShouldTransact)
 {
+#if WITH_EDITOR
 	TRACE_CPUPROFILER_EVENT_SCOPE(UAnimDatabaseLibrary::DatabasePatchRanges);
 
 	if (!Database || !Database->GetSkeleton()) { return; }
@@ -6544,10 +6601,12 @@ void UAnimDatabaseLibrary::DatabaseSetToReferencePose(UAnimDatabase* Database, c
 			}
 
 		}, bShouldTransact);
+#endif
 }
 
 void UAnimDatabaseLibrary::DatabaseRemoveFootGroundPenetration(UAnimDatabase* Database, const FAnimDatabaseFrameRanges& FrameRanges, const FName LeftToeBone, const FName RightToeBone, const float ToeBoneLength, const float PelvisHeightAdjustment, const FVector LeftKneeSideVector, const FVector RightKneeSideVector, const FVector LeftToeForwardVector, const FVector RightToeForwardVector, const bool bShouldTransact)
 {
+#if WITH_EDITOR
 	if (!Database || !Database->GetSkeleton()) { return; }
 
 	const int32 LeftToeBoneIndex = Database->FindBoneIndex(LeftToeBone);
@@ -6635,8 +6694,11 @@ void UAnimDatabaseLibrary::DatabaseRemoveFootGroundPenetration(UAnimDatabase* Da
 				ContactSettings);
 
 		}, bShouldTransact);
+#endif
 }
 
+#if WITH_EDITOR
+
 void UAnimDatabaseFunction_StripUnusedBoneTracks::Run_Implementation(UAnimDatabase* InDatabase, const FAnimDatabaseFrameRanges& InFrameRanges, const bool bShouldTransact)
 {
 	UAnimDatabaseLibrary::DatabaseStripUnusedBoneTracks(InDatabase, InFrameRanges, bShouldTransact);
@@ -6877,6 +6939,58 @@ void UAnimDatabaseFunction_AddCurveFromFrameAttribute::Run_Implementation(UAnimD
 		bShouldTransact);
 }
 
+void  UAnimDatabaseFunction_AddContactCurves::Run_Implementation(UAnimDatabase* InDatabase, const FAnimDatabaseFrameRanges& InFrameRanges, const bool bShouldTransact)
+{
+	UE::Tasks::TTask<FAnimDatabaseFrameAttribute> LeftContactCurve = UE::Tasks::Launch(UE_SOURCE_LOCATION, [this, InDatabase, InFrameRanges]() {
+		
+		return UAnimDatabaseFrameAttributeLibrary::FrameAttributeContactCurve(
+			InDatabase,
+			InFrameRanges,
+			LeftContactBoneName,
+			HeightThreshold,
+			VelocityThreshold,
+			MajorityVoteFilterWidth,
+			SmoothingAmount);
+		});
+
+	UE::Tasks::TTask<FAnimDatabaseFrameAttribute> RightContactCurve = UE::Tasks::Launch(UE_SOURCE_LOCATION, [this, InDatabase, InFrameRanges]() {
+
+		return UAnimDatabaseFrameAttributeLibrary::FrameAttributeContactCurve(
+			InDatabase,
+			InFrameRanges,
+			RightContactBoneName,
+			HeightThreshold,
+			VelocityThreshold,
+			MajorityVoteFilterWidth,
+			SmoothingAmount);
+		});
+
+	UAnimDatabaseLibrary::DatabaseAddCurvesFromArrayViews(
+		InDatabase,
+		{ LeftContactCurveName, RightContactCurveName },
+		{ LeftContactCurve.GetResult(), RightContactCurve.GetResult() },
+		bSparseKeys,
+		bShouldTransact);
+}
+
+void UAnimDatabaseFunction_SetBoneLocalTransformFromFrameAttribute::Run_Implementation(UAnimDatabase* InDatabase, const FAnimDatabaseFrameRanges& InFrameRanges, const bool bShouldTransact)
+{
+	UAnimDatabaseLibrary::DatabaseSetLocalBoneTransform(InDatabase, BoneName, UAnimDatabaseFrameAttributeLibrary::MakeFrameAttributeFromFunction(InDatabase, InFrameRanges, FrameAttribute), bShouldTransact);
+}
+
+void UAnimDatabaseFunction_SetBoneGlobalTransformFromFrameAttribute::Run_Implementation(UAnimDatabase* InDatabase, const FAnimDatabaseFrameRanges& InFrameRanges, const bool bShouldTransact)
+{
+	const int32 BoneIndex = InDatabase->FindBoneIndex(BoneName);
+	if (BoneIndex == INDEX_NONE) { return; }
+	const int32 ParentIndex = InDatabase->GetBoneParent(BoneIndex);
+	if (ParentIndex == INDEX_NONE) { return; }
+
+	UAnimDatabaseLibrary::DatabaseSetLocalBoneTransform(InDatabase, BoneName, 
+		UAnimDatabaseFrameAttributeLibrary::FrameAttributeDivide(
+			UAnimDatabaseFrameAttributeLibrary::MakeFrameAttributeFromFunction(InDatabase, InFrameRanges, FrameAttribute),
+			UAnimDatabaseFrameAttributeLibrary::MakeBoneGlobalTransformFrameAttribute(InDatabase, InFrameRanges, ParentIndex)), bShouldTransact);
+}
+
 void UAnimDatabaseFunction_SetBoneToRootTransformAtFrames::Run_Implementation(UAnimDatabase* InDatabase, const FAnimDatabaseFrameRanges& InFrameRanges, const bool bShouldTransact)
 {
 	const int32 BoneIndex = InDatabase->FindBoneIndex(BoneName);
@@ -7108,6 +7222,29 @@ void UAnimDatabaseDebugDraw_RootOrientation::DrawDebug_Implementation(
 		Radius);
 }
 
+void UAnimDatabaseDebugDraw_BoneTransform::DrawDebug_Implementation(
+	const FDebugDrawer& Drawer,
+	const FAnimDatabasePoseState& InPoseState,
+	const UAnimDatabase* InDatabase,
+	const FAnimDatabaseFrameRanges& InFrameRanges,
+	const FTransform& RangeViewportTransform,
+	const int32 SequenceIdx,
+	const float SequenceTime,
+	const int32 RangeStart,
+	const int32 RangeLength,
+	const FLinearColor& IdentifierColor)
+{
+	const int32 BoneIndex = InDatabase->FindBoneIndex(BoneName);
+	if (BoneIndex == INDEX_NONE) { return; }
+
+	UDrawDebugLibrary::DrawDebugTransform(
+		Drawer,
+		UAnimDatabasePoseStateLibrary::PoseStateBoneWorldTransform(InPoseState, BoneIndex),
+		UDrawDebugLibrary::MakeDrawDebugLineStyleFromColorAndThickness(IdentifierColor, Thickness),
+		bDepthTest,
+		Radius);
+}
+
 void UAnimDatabaseDebugDraw_RootVelocity::InitializeDrawDebug_Implementation(const UAnimDatabase* InDatabase, const FAnimDatabaseFrameRanges& InFrameRanges)
 {
 	RootLinearVelocity = UAnimDatabaseFrameAttributeLibrary::FrameAttributeLength(UAnimDatabaseFrameAttributeLibrary::MakeRootLinearVelocityFrameAttribute(InDatabase, InFrameRanges));
diff --git a/Engine/Plugins/Experimental/Animation/AnimDatabase/Source/AnimDatabase/Private/AnimDatabaseFrameAttribute.cpp b/Engine/Plugins/Experimental/Animation/AnimDatabase/Source/AnimDatabase/Private/AnimDatabaseFrameAttribute.cpp
--- a/Engine/Plugins/Experimental/Animation/AnimDatabase/Source/AnimDatabase/Private/AnimDatabaseFrameAttribute.cpp
+++ b/Engine/Plugins/Experimental/Animation/AnimDatabase/Source/AnimDatabase/Private/AnimDatabaseFrameAttribute.cpp
@@ -464,6 +464,31 @@ FAnimDatabaseFrameAttribute UAnimDatabaseFrameAttributeLibrary::MakeFrameAttribu
 	return OutFrameAttribute;
 }
 
+void UAnimDatabaseFrameAttributeLibrary::MakeFrameAttributesFromFunctions(
+	TArray<FAnimDatabaseFrameAttribute>& OutFrameAttributes,
+	const UAnimDatabase* Database,
+	const FAnimDatabaseFrameRanges& FrameRanges,
+	const TArray<UAnimDatabaseFrameAttributeFunction*>& Functions)
+{
+	OutFrameAttributes.SetNum(Functions.Num());
+	MakeFrameAttributesFromFunctionsArrayView(OutFrameAttributes, Database, FrameRanges, Functions);
+}
+
+void UAnimDatabaseFrameAttributeLibrary::MakeFrameAttributesFromFunctionsArrayView(
+	const TArrayView<FAnimDatabaseFrameAttribute> OutFrameAttributes,
+	const UAnimDatabase* Database,
+	const FAnimDatabaseFrameRanges& FrameRanges,
+	const TArrayView<UAnimDatabaseFrameAttributeFunction* const> Functions)
+{
+	check(OutFrameAttributes.Num() == Functions.Num());
+	const int32 FunctionNum = Functions.Num();
+
+	for (int32 FunctionIdx = 0; FunctionIdx < FunctionNum; FunctionIdx++)
+	{
+		OutFrameAttributes[FunctionIdx] = MakeFrameAttributeFromFunction(Database, FrameRanges, Functions[FunctionIdx]);
+	}
+}
+
 FAnimDatabaseFrameAttribute UAnimDatabaseFrameAttributeLibrary::MakeFrameAttributeFromClass(const UAnimDatabase* Database, const FAnimDatabaseFrameRanges& FrameRanges, const TSubclassOf<UAnimDatabaseFrameAttributeFunction> Class)
 {
 	if (!Database)
@@ -1252,6 +1277,29 @@ FAnimDatabaseFrameAttribute UAnimDatabaseFrameAttributeLibrary::MakeRotationFram
 	return Out;
 }
 
+FAnimDatabaseFrameAttribute UAnimDatabaseFrameAttributeLibrary::MakeRotationFrameAttributeFromBasisDirections(const FAnimDatabaseFrameAttribute& ForwardDirection, const FAnimDatabaseFrameAttribute& RightDirection, const FAnimDatabaseFrameAttribute& UpDirection)
+{
+	if (!ForwardDirection.IsValid() || !RightDirection.IsValid() || !UpDirection.IsValid())
+	{
+		UE_LOGFMT(LogAnimDatabase, Error, "MakeRotationFrameAttributeFromBasisDirections: Invalid Frame Attribute.");
+		return FAnimDatabaseFrameAttribute();
+	}
+
+	if (ForwardDirection.Type != EAnimDatabaseAttributeType::Direction ||
+		RightDirection.Type != EAnimDatabaseAttributeType::Direction ||
+		UpDirection.Type != EAnimDatabaseAttributeType::Direction)
+	{
+		UE_LOGFMT(LogAnimDatabase, Error, "MakeRotationFrameAttributeFromBasisDirections: FrameAttribute Types don't match. Got {Type}, {Type}, and {Type}. Expected Directions.", 
+			AttributeTypeNameInternal(ForwardDirection.Type), AttributeTypeNameInternal(RightDirection.Type), AttributeTypeNameInternal(UpDirection.Type));
+		return FAnimDatabaseFrameAttribute();
+	}
+
+	FAnimDatabaseFrameAttribute Out = MakeEmptyFrameAttribute();
+	Out.Type = EAnimDatabaseAttributeType::Rotation;
+	UE::Learning::FrameAttribute::QuatFromBasisDirections(*Out.FrameAttribute, *ForwardDirection.FrameAttribute, *RightDirection.FrameAttribute, *UpDirection.FrameAttribute);
+	return Out;
+}
+
 FAnimDatabaseFrameAttribute UAnimDatabaseFrameAttributeLibrary::MakeScaleFrameAttributeFromConstant(const FAnimDatabaseFrameRanges& FrameRanges, const FVector Scale)
 {
 	if (!FrameRanges.IsValid())
@@ -7242,6 +7290,50 @@ FAnimDatabaseFrameAttribute UAnimDatabaseFrameAttributeLibrary::FrameAttributeTr
 	return Out;
 }
 
+FAnimDatabaseFrameAttribute UAnimDatabaseFrameAttributeLibrary::FrameAttributeApplyTransformLocation(const FAnimDatabaseFrameAttribute& A, const FVector& B)
+{
+	TRACE_CPUPROFILER_EVENT_SCOPE(UAnimDatabaseFrameAttributeLibrary::FrameAttributeApplyTransformLocation);
+
+	if (!A.IsValid())
+	{
+		UE_LOGFMT(LogAnimDatabase, Error, "FrameAttributeApplyTransformLocation: Invalid FrameAttribute.");
+		return FAnimDatabaseFrameAttribute();
+	}
+
+	if (A.Type != EAnimDatabaseAttributeType::Transform)
+	{
+		UE_LOGFMT(LogAnimDatabase, Error, "FrameAttributeApplyTransformLocation: FrameAttribute Type not valid. Got {Type}. Expected Transform.", AttributeTypeNameInternal(A.Type));
+		return FAnimDatabaseFrameAttribute();
+	}
+
+	FAnimDatabaseFrameAttribute Out = MakeEmptyFrameAttribute();
+	Out.Type = EAnimDatabaseAttributeType::Location;
+	UE::Learning::FrameAttribute::TransformApplyConstantLocation(*Out.FrameAttribute, *A.FrameAttribute, (FVector3f)B);
+	return Out;
+}
+
+FAnimDatabaseFrameAttribute UAnimDatabaseFrameAttributeLibrary::FrameAttributeApplyTransformDirection(const FAnimDatabaseFrameAttribute& A, const FVector& B)
+{
+	TRACE_CPUPROFILER_EVENT_SCOPE(UAnimDatabaseFrameAttributeLibrary::FrameAttributeApplyTransformDirection);
+
+	if (!A.IsValid())
+	{
+		UE_LOGFMT(LogAnimDatabase, Error, "FrameAttributeApplyTransformDirection: Invalid FrameAttribute.");
+		return FAnimDatabaseFrameAttribute();
+	}
+
+	if (A.Type != EAnimDatabaseAttributeType::Transform)
+	{
+		UE_LOGFMT(LogAnimDatabase, Error, "FrameAttributeApplyTransformDirection: FrameAttribute Type not valid. Got {Type}. Expected Transform.", AttributeTypeNameInternal(A.Type));
+		return FAnimDatabaseFrameAttribute();
+	}
+
+	FAnimDatabaseFrameAttribute Out = MakeEmptyFrameAttribute();
+	Out.Type = EAnimDatabaseAttributeType::Direction;
+	UE::Learning::FrameAttribute::TransformApplyConstantDirection(*Out.FrameAttribute, *A.FrameAttribute, (FVector3f)B);
+	return Out;
+}
+
 FAnimDatabaseFrameAttribute UAnimDatabaseFrameAttributeLibrary::FrameAttributeAddAngleDegrees(const FAnimDatabaseFrameAttribute& A, const float B)
 {
 	TRACE_CPUPROFILER_EVENT_SCOPE(UAnimDatabaseFrameAttributeLibrary::FrameAttributeAddAngleDegrees);
@@ -8413,6 +8505,46 @@ FAnimDatabaseFrameAttribute UAnimDatabaseFrameAttributeLibrary::FrameAttributeFl
 	return Out;
 }
 
+FAnimDatabaseFrameAttribute UAnimDatabaseFrameAttributeLibrary::FrameAttributeDirectionToLocation(const FAnimDatabaseFrameAttribute& DirectionFrameAttribute)
+{
+	if (!DirectionFrameAttribute.IsValid())
+	{
+		UE_LOGFMT(LogAnimDatabase, Error, "FrameAttributeDirectionToLocation: Invalid FrameAttribute.");
+		return FAnimDatabaseFrameAttribute();
+	}
+
+	if (DirectionFrameAttribute.Type != EAnimDatabaseAttributeType::Direction)
+	{
+		UE_LOGFMT(LogAnimDatabase, Error, "FrameAttributeDirectionToLocation: FrameAttribute Types don't match. Got {Type}. Expected Direction.", AttributeTypeNameInternal(DirectionFrameAttribute.Type));
+		return FAnimDatabaseFrameAttribute();
+	}
+
+	// We can just return it directly with a different type the underlying data doesn't need changing
+	FAnimDatabaseFrameAttribute Out = DirectionFrameAttribute;
+	Out.Type = EAnimDatabaseAttributeType::Location;
+	return Out;
+}
+
+FAnimDatabaseFrameAttribute UAnimDatabaseFrameAttributeLibrary::FrameAttributeLocationToDirection(const FAnimDatabaseFrameAttribute& LocationFrameAttribute)
+{
+	if (!LocationFrameAttribute.IsValid())
+	{
+		UE_LOGFMT(LogAnimDatabase, Error, "FrameAttributeLocationToDirection: Invalid FrameAttribute.");
+		return FAnimDatabaseFrameAttribute();
+	}
+
+	if (LocationFrameAttribute.Type != EAnimDatabaseAttributeType::Location)
+	{
+		UE_LOGFMT(LogAnimDatabase, Error, "FrameAttributeLocationToDirection: FrameAttribute Types don't match. Got {Type}. Expected Location.", AttributeTypeNameInternal(LocationFrameAttribute.Type));
+		return FAnimDatabaseFrameAttribute();
+	}
+
+	FAnimDatabaseFrameAttribute Out = FrameAttributeCopy(LocationFrameAttribute);
+	Out.Type = EAnimDatabaseAttributeType::Direction;
+	UE::Learning::FrameAttribute::NormalizeInplace(*Out.FrameAttribute);
+	return Out;
+}
+
 FAnimDatabaseFrameAttribute UAnimDatabaseFrameAttributeLibrary::FrameAttributeFilterGaussian(const FAnimDatabaseFrameAttribute& A, const float StandardDeviationInFrames)
 {
 	if (!A.IsValid())
@@ -8674,6 +8806,43 @@ FAnimDatabaseFrameAttribute UAnimDatabaseFrameAttributeLibrary::FrameAttributeSu
 	return Out;
 }
 
+FAnimDatabaseFrameAttribute UAnimDatabaseFrameAttributeLibrary::FrameAttributeWeightedSum(const TArray<FAnimDatabaseFrameAttribute>& FrameAttributes, const TArray<float>& Weights)
+{
+	return FrameAttributeWeightedSumFromArrayView(FrameAttributes, Weights);
+}
+
+FAnimDatabaseFrameAttribute UAnimDatabaseFrameAttributeLibrary::FrameAttributeWeightedSumFromArrayView(const TArrayView<const FAnimDatabaseFrameAttribute> FrameAttributes, const TArrayView<const float> Weights)
+{
+	const int32 FrameAttributeNum = FrameAttributes.Num();
+	const int32 WeightNum = Weights.Num();
+
+	if (FrameAttributeNum == 0)
+	{
+		UE_LOGFMT(LogAnimDatabase, Warning, "FrameAttributeWeightedSumFromArrayView: No FrameAttributes Provided.");
+		return FAnimDatabaseFrameAttribute();
+	}
+
+	if (WeightNum != FrameAttributeNum)
+	{
+		UE_LOGFMT(LogAnimDatabase, Error, "FrameAttributeWeightedSumFromArrayView: Number of Frame Attributes {NumAtr} does not match number of weights {NumWeight}.",
+			FrameAttributeNum, WeightNum);
+		return FAnimDatabaseFrameAttribute();
+	}
+
+	if (FrameAttributeNum == 1)
+	{
+		return FrameAttributeMultiplyFloat(FrameAttributes[0], Weights[0]);
+	}
+
+	// TODO: Optimize this properly
+	FAnimDatabaseFrameAttribute Out = FrameAttributeMultiplyFloat(FrameAttributes[0], Weights[0]);
+	for (int32 FrameAttributeIdx = 1; FrameAttributeIdx < FrameAttributeNum; FrameAttributeIdx++)
+	{
+		Out = FrameAttributeAdd(Out, FrameAttributeMultiplyFloat(FrameAttributes[FrameAttributeIdx], Weights[FrameAttributeIdx]));
+	}
+	return Out;
+}
+
 FAnimDatabaseFrameAttribute UAnimDatabaseFrameAttributeLibrary::FrameAttributeProduct(const TArray<FAnimDatabaseFrameAttribute>& FrameAttributes)
 {
 	return FrameAttributeProductFromArrayView(FrameAttributes);
@@ -9509,6 +9678,124 @@ FAnimDatabaseFrameAttribute UAnimDatabaseFrameAttributeLibrary::FrameAttributeCo
 		Database->GetFrameRate());
 }
 
+FAnimDatabaseFrameAttribute UAnimDatabaseFrameAttributeLibrary::FrameAttributeFitTransform(
+	const UAnimDatabase* Database,
+	const FAnimDatabaseFrameRanges& FrameRanges,
+	const TArray<FAnimDatabaseFrameAttribute>& FitLocations,
+	const TArray<FAnimDatabaseFrameAttribute>& FitForwardDirections,
+	const TArray<FAnimDatabaseFrameAttribute>& FitRightDirections,
+	const TArray<FAnimDatabaseFrameAttribute>& FitUpDirections,
+	const TArray<float>& FitLocationWeights,
+	const TArray<float>& FitForwardDirectionWeights,
+	const TArray<float>& FitRightDirectionWeights,
+	const TArray<float>& FitUpDirectionWeights)
+{
+	return FrameAttributeFitTransformFromArrayViews(
+		Database,
+		FrameRanges,
+		FitLocations,
+		FitForwardDirections,
+		FitRightDirections,
+		FitUpDirections,
+		FitLocationWeights,
+		FitForwardDirectionWeights,
+		FitRightDirectionWeights,
+		FitUpDirectionWeights);
+}
+
+FAnimDatabaseFrameAttribute UAnimDatabaseFrameAttributeLibrary::FrameAttributeFitTransformFromArrayViews(
+	const UAnimDatabase* Database,
+	const FAnimDatabaseFrameRanges& FrameRanges,
+	const TArrayView<const FAnimDatabaseFrameAttribute> FitLocations,
+	const TArrayView<const FAnimDatabaseFrameAttribute> FitForwardDirections,
+	const TArrayView<const FAnimDatabaseFrameAttribute> FitRightDirections,
+	const TArrayView<const FAnimDatabaseFrameAttribute> FitUpDirections,
+	const TArrayView<const float> FitLocationWeights,
+	const TArrayView<const float> FitForwardDirectionWeights,
+	const TArrayView<const float> FitRightDirectionWeights,
+	const TArrayView<const float> FitUpDirectionWeights)
+{
+	const int32 LocationNum = FitLocations.Num();
+	const int32 ForwardNum = FitForwardDirections.Num();
+	const int32 RightNum = FitRightDirections.Num();
+	const int32 UpNum = FitUpDirections.Num();
+
+	if (LocationNum == 0 || ForwardNum == 0 || RightNum == 0 || UpNum == 0) { return FAnimDatabaseFrameAttribute(); }
+
+	TArray<FAnimDatabaseFrameAttribute, TInlineAllocator<32>> ForwardAttributes;
+	TArray<FAnimDatabaseFrameAttribute, TInlineAllocator<32>> RightAttributes;
+	TArray<FAnimDatabaseFrameAttribute, TInlineAllocator<32>> UpAttributes;
+
+	ForwardAttributes.Reserve(ForwardNum);
+	RightAttributes.Reserve(RightNum);
+	UpAttributes.Reserve(UpNum);
+
+	TArray<float> LocationWeights; LocationWeights.SetNumZeroed(LocationNum);
+	TArray<float> ForwardWeights; ForwardWeights.SetNumZeroed(ForwardNum);
+	TArray<float> RightWeights; RightWeights.SetNumZeroed(RightNum);
+	TArray<float> UpWeights; UpWeights.SetNumZeroed(UpNum);
+
+	float LocationWeightSum = 0.0f;
+	for (int32 LocationIdx = 0; LocationIdx < LocationNum; LocationIdx++)
+	{
+		LocationWeights[LocationIdx] = LocationIdx < FitLocationWeights.Num() ? FitLocationWeights[LocationIdx] : 1.0f;
+		LocationWeightSum += LocationWeights[LocationIdx];
+	}
+
+	for (int32 LocationIdx = 0; LocationIdx < LocationNum; LocationIdx++)
+	{
+		LocationWeights[LocationIdx] /= FMath::Max(LocationWeightSum, UE_SMALL_NUMBER);
+	}
+
+	float ForwardWeightSum = 0.0f;
+	for (int32 DirectionIdx = 0; DirectionIdx < ForwardNum; DirectionIdx++)
+	{
+		ForwardWeights[DirectionIdx] = DirectionIdx < FitForwardDirectionWeights.Num() ? FitForwardDirectionWeights[DirectionIdx] : 1.0f;
+		ForwardWeightSum += ForwardWeights[DirectionIdx];
+		ForwardAttributes.Add(UAnimDatabaseFrameAttributeLibrary::FrameAttributeDirectionToLocation(FitForwardDirections[DirectionIdx]));
+	}
+
+	for (int32 DirectionIdx = 0; DirectionIdx < ForwardNum; DirectionIdx++)
+	{
+		ForwardWeights[DirectionIdx] /= FMath::Max(ForwardWeightSum, UE_SMALL_NUMBER);
+	}
+
+	float RightWeightSum = 0.0f;
+	for (int32 DirectionIdx = 0; DirectionIdx < RightNum; DirectionIdx++)
+	{
+		RightWeights[DirectionIdx] = DirectionIdx < FitRightDirectionWeights.Num() ? FitRightDirectionWeights[DirectionIdx] : 1.0f;
+		RightWeightSum += RightWeights[DirectionIdx];
+		RightAttributes.Add(UAnimDatabaseFrameAttributeLibrary::FrameAttributeDirectionToLocation(FitRightDirections[DirectionIdx]));
+	}
+
+	for (int32 DirectionIdx = 0; DirectionIdx < RightNum; DirectionIdx++)
+	{
+		RightWeights[DirectionIdx] /= FMath::Max(RightWeightSum, UE_SMALL_NUMBER);
+	}
+
+	float UpWeightSum = 0.0f;
+	for (int32 DirectionIdx = 0; DirectionIdx < UpNum; DirectionIdx++)
+	{
+		UpWeights[DirectionIdx] = DirectionIdx < FitUpDirectionWeights.Num() ? FitUpDirectionWeights[DirectionIdx] : 1.0f;
+		UpWeightSum += UpWeights[DirectionIdx];
+		UpAttributes.Add(UAnimDatabaseFrameAttributeLibrary::FrameAttributeDirectionToLocation(FitUpDirections[DirectionIdx]));
+	}
+
+	for (int32 DirectionIdx = 0; DirectionIdx < UpNum; DirectionIdx++)
+	{
+		UpWeights[DirectionIdx] /= FMath::Max(UpWeightSum, UE_SMALL_NUMBER);
+	}
+
+	const FAnimDatabaseFrameAttribute LocationAverage = UAnimDatabaseFrameAttributeLibrary::FrameAttributeWeightedSumFromArrayView(FitLocations, LocationWeights);
+	const FAnimDatabaseFrameAttribute ForwardAverage = UAnimDatabaseFrameAttributeLibrary::FrameAttributeLocationToDirection(UAnimDatabaseFrameAttributeLibrary::FrameAttributeWeightedSumFromArrayView(ForwardAttributes, ForwardWeights));
+	const FAnimDatabaseFrameAttribute RightAverage = UAnimDatabaseFrameAttributeLibrary::FrameAttributeLocationToDirection(UAnimDatabaseFrameAttributeLibrary::FrameAttributeWeightedSumFromArrayView(RightAttributes, RightWeights));
+	const FAnimDatabaseFrameAttribute UpAverage = UAnimDatabaseFrameAttributeLibrary::FrameAttributeLocationToDirection(UAnimDatabaseFrameAttributeLibrary::FrameAttributeWeightedSumFromArrayView(UpAttributes, UpWeights));
+
+	const FAnimDatabaseFrameAttribute RotationAverage = UAnimDatabaseFrameAttributeLibrary::MakeRotationFrameAttributeFromBasisDirections(ForwardAverage, RightAverage, UpAverage);
+
+	return UAnimDatabaseFrameAttributeLibrary::MakeTransformFrameAttributeNoScale(LocationAverage, RotationAverage);
+}
+
 FAnimDatabaseFrameAttribute UAnimDatabaseFrameAttributeFunction::MakeFrameAttribute_Implementation(const UAnimDatabase* Database, const FAnimDatabaseFrameRanges& FrameRanges) const
 {
 	return UAnimDatabaseFrameAttributeLibrary::MakeEmptyFrameAttribute();
@@ -9543,6 +9830,45 @@ FAnimDatabaseFrameAttribute UAnimDatabaseFrameAttributeFunction_ContactCurve::Ma
 		SmoothingAmount);
 }
 
+FAnimDatabaseFrameAttribute UAnimDatabaseFrameAttributeFunction_FitTransform::MakeFrameAttribute_Implementation(const UAnimDatabase* Database, const FAnimDatabaseFrameRanges& FrameRanges) const
+{
+	TArray<FAnimDatabaseFrameAttribute> LocationAttributes;
+	TArray<FAnimDatabaseFrameAttribute> ForwardAttributes;
+	TArray<FAnimDatabaseFrameAttribute> RightAttributes;
+	TArray<FAnimDatabaseFrameAttribute> UpAttributes;
+	UAnimDatabaseFrameAttributeLibrary::MakeFrameAttributesFromFunctions(LocationAttributes, Database, FrameRanges, FitLocations);
+	UAnimDatabaseFrameAttributeLibrary::MakeFrameAttributesFromFunctions(ForwardAttributes, Database, FrameRanges, FitForwardDirections);
+	UAnimDatabaseFrameAttributeLibrary::MakeFrameAttributesFromFunctions(RightAttributes, Database, FrameRanges, FitRightDirections);
+	UAnimDatabaseFrameAttributeLibrary::MakeFrameAttributesFromFunctions(UpAttributes, Database, FrameRanges, FitUpDirections);
+
+	return UAnimDatabaseFrameAttributeLibrary::FrameAttributeFitTransformFromArrayViews(
+		Database,
+		FrameRanges,
+		LocationAttributes,
+		ForwardAttributes,
+		RightAttributes,
+		UpAttributes,
+		FitLocationWeights,
+		FitForwardDirectionWeights,
+		FitRightDirectionWeights,
+		FitUpDirectionWeights);
+}
+
+FAnimDatabaseFrameAttribute UAnimDatabaseFrameAttributeFunction_BoneGlobalTransform::MakeFrameAttribute_Implementation(const UAnimDatabase* Database, const FAnimDatabaseFrameRanges& FrameRanges) const
+{
+	return UAnimDatabaseFrameAttributeLibrary::FrameAttributeTransformMultiply(LocalOffset, UAnimDatabaseFrameAttributeLibrary::MakeBoneGlobalTransformFrameAttributeFromName(Database, FrameRanges, BoneName));
+}
+
+FAnimDatabaseFrameAttribute UAnimDatabaseFrameAttributeFunction_BoneGlobalLocation::MakeFrameAttribute_Implementation(const UAnimDatabase* Database, const FAnimDatabaseFrameRanges& FrameRanges) const
+{
+	return UAnimDatabaseFrameAttributeLibrary::FrameAttributeApplyTransformLocation(UAnimDatabaseFrameAttributeLibrary::MakeBoneGlobalTransformFrameAttributeFromName(Database, FrameRanges, BoneName), LocalOffset);
+}
+
+FAnimDatabaseFrameAttribute UAnimDatabaseFrameAttributeFunction_BoneGlobalDirection::MakeFrameAttribute_Implementation(const UAnimDatabase* Database, const FAnimDatabaseFrameRanges& FrameRanges) const
+{
+	return UAnimDatabaseFrameAttributeLibrary::MakeBoneGlobalDirectionFrameAttributeFromName(Database, FrameRanges, BoneName, LocalDirection);
+}
+
 FAnimDatabaseFrameAttribute UAnimDatabaseFrameAttributeFunction_BoneGlobalLinearVelocityMagnitude::MakeFrameAttribute_Implementation(const UAnimDatabase* Database, const FAnimDatabaseFrameRanges& FrameRanges) const
 {
 	return UAnimDatabaseFrameAttributeLibrary::FrameAttributeBoneSpeed(Database, FrameRanges, BoneName);
diff --git a/Engine/Plugins/Experimental/Animation/AnimDatabase/Source/AnimDatabase/Private/AnimDatabaseMath.cpp b/Engine/Plugins/Experimental/Animation/AnimDatabase/Source/AnimDatabase/Private/AnimDatabaseMath.cpp
--- a/Engine/Plugins/Experimental/Animation/AnimDatabase/Source/AnimDatabase/Private/AnimDatabaseMath.cpp
+++ b/Engine/Plugins/Experimental/Animation/AnimDatabase/Source/AnimDatabase/Private/AnimDatabaseMath.cpp
@@ -3205,6 +3205,20 @@ namespace UE::AnimDatabase::Math
 
 	// Mean/Std/Scale
 
+	float Sum(const TLearningArrayView<1, const float> Data)
+	{
+		TRACE_CPUPROFILER_EVENT_SCOPE(UE::AnimDatabase::Math::Sum);
+
+		const int64 Num = Data.Num();
+
+		float Total = 0.0f;
+		for (int64 Idx = 0; Idx < Num; Idx++)
+		{
+			Total += Data[Idx];
+		}
+		return Total;
+	}
+
 	void ComputeMean(float& OutMean, const TLearningArrayView<1, const float> Data)
 	{
 		TRACE_CPUPROFILER_EVENT_SCOPE(UE::AnimDatabase::Math::ComputeMean);
diff --git a/Engine/Plugins/Experimental/Animation/AnimDatabase/Source/AnimDatabase/Public/AnimDatabase.h b/Engine/Plugins/Experimental/Animation/AnimDatabase/Source/AnimDatabase/Public/AnimDatabase.h
--- a/Engine/Plugins/Experimental/Animation/AnimDatabase/Source/AnimDatabase/Public/AnimDatabase.h
+++ b/Engine/Plugins/Experimental/Animation/AnimDatabase/Source/AnimDatabase/Public/AnimDatabase.h
@@ -879,8 +879,6 @@ class UAnimDatabaseLibrary : public UBlueprintFunctionLibrary
 {
 	GENERATED_BODY()
 
-#if WITH_EDITOR
-
 public:
 
 	/** Gets all of the AnimNotify classes present in a given AnimSequence */
@@ -1026,8 +1024,6 @@ private:
 
 	/** Loads into memory all of the blueprint assets of a given class. */
 	static void LoadAllBlueprintAssetsOfClass(UClass* Class);
-
-#endif
 };
 
 /** Function for stripping unused bone-tracks from the AnimSequence to reduce the file size on desk  */
@@ -1674,6 +1670,119 @@ public:
 #endif
 };
 
+/** Function for adding contact curves  */
+UCLASS(BlueprintType, Blueprintable, meta = (DisplayName = "Add Contact Curves"))
+class UAnimDatabaseFunction_AddContactCurves : public UAnimDatabaseFunction
+{
+	GENERATED_BODY()
+
+#if WITH_EDITORONLY_DATA
+
+public:
+
+	/** Name of the contact curve for the left foot */
+	UPROPERTY(EditAnywhere, Category = "Settings")
+	FName LeftContactCurveName = TEXT("contact_l");
+
+	/** Name of the contact curve for the right foot */
+	UPROPERTY(EditAnywhere, Category = "Settings")
+	FName RightContactCurveName = TEXT("contact_r");
+
+	/** Name of the left bone to generate the contact curve for */
+	UPROPERTY(EditAnywhere, Category = "Settings")
+	FName LeftContactBoneName = TEXT("ball_l");
+
+	/** Name of the right bone to generate the contact curve for */
+	UPROPERTY(EditAnywhere, Category = "Settings")
+	FName RightContactBoneName = TEXT("ball_r");
+
+	/** Maximum height (in cm) for which the bone can be considered in contact */
+	UPROPERTY(EditAnywhere, Category = "Settings")
+	float HeightThreshold = 20.0f;
+
+	/** Maximum velocity (in cm/s) for which the bone can be considered in contact */
+	UPROPERTY(EditAnywhere, Category = "Settings")
+	float VelocityThreshold = 25.0f;
+
+	/**
+	 * Filter width (in frames) for the majority vote filter used to remove very short contact transitions.
+	 * Contact periods (or non-contact periods) less than half this duration will be removed.
+	 */
+	UPROPERTY(EditAnywhere, Category = "Settings")
+	int32 MajorityVoteFilterWidth = 5;
+
+	/* Amount of smoothing to apply to the resulting contact curve (in frames) */
+	UPROPERTY(EditAnywhere, Category = "Settings")
+	float SmoothingAmount = 3.0f;
+
+	UPROPERTY(EditAnywhere, Category = "Settings")
+	bool bSparseKeys = true;
+
+#endif
+
+#if WITH_EDITOR
+
+public:
+
+	UE_API virtual void Run_Implementation(UAnimDatabase* InDatabase, const FAnimDatabaseFrameRanges& InFrameRanges, const bool bShouldTransact) override;
+
+#endif
+};
+
+/** Function for setting a bone local transform from frame attribute function  */
+UCLASS(BlueprintType, Blueprintable, meta = (DisplayName = "Set Bone Local Transform from Frame Attribute"))
+class UAnimDatabaseFunction_SetBoneLocalTransformFromFrameAttribute : public UAnimDatabaseFunction
+{
+	GENERATED_BODY()
+
+#if WITH_EDITORONLY_DATA
+
+public:
+
+	UPROPERTY(EditAnywhere, Category = "Settings")
+	FName BoneName = NAME_None;
+
+	UPROPERTY(EditAnywhere, Instanced, Category = "Settings")
+	TObjectPtr<UAnimDatabaseFrameAttributeFunction> FrameAttribute;
+
+#endif
+
+#if WITH_EDITOR
+
+public:
+
+	UE_API virtual void Run_Implementation(UAnimDatabase* InDatabase, const FAnimDatabaseFrameRanges& InFrameRanges, const bool bShouldTransact) override;
+
+#endif
+};
+
+/** Function for setting a bone global transform from frame attribute function  */
+UCLASS(BlueprintType, Blueprintable, meta = (DisplayName = "Set Bone Global Transform from Frame Attribute"))
+class UAnimDatabaseFunction_SetBoneGlobalTransformFromFrameAttribute : public UAnimDatabaseFunction
+{
+	GENERATED_BODY()
+
+#if WITH_EDITORONLY_DATA
+
+public:
+
+	UPROPERTY(EditAnywhere, Category = "Settings")
+	FName BoneName = NAME_None;
+
+	UPROPERTY(EditAnywhere, Instanced, Category = "Settings")
+	TObjectPtr<UAnimDatabaseFrameAttributeFunction> FrameAttribute;
+
+#endif
+
+#if WITH_EDITOR
+
+public:
+
+	UE_API virtual void Run_Implementation(UAnimDatabase* InDatabase, const FAnimDatabaseFrameRanges& InFrameRanges, const bool bShouldTransact) override;
+
+#endif
+};
+
 /** Function to set bone to root location at frame  */
 UCLASS(BlueprintType, Blueprintable, meta = (DisplayName = "Set Bone to Root Transform at Frames"))
 class UAnimDatabaseFunction_SetBoneToRootTransformAtFrames : public UAnimDatabaseFunction
@@ -2121,6 +2230,53 @@ public:
 #endif
 };
 
+/** Debug Draw the transform of the given bone */
+UCLASS(BlueprintType, Blueprintable, meta = (DisplayName = "Bone Transform"))
+class UAnimDatabaseDebugDraw_BoneTransform : public UAnimDatabaseDebugDraw
+{
+	GENERATED_BODY()
+
+#if WITH_EDITORONLY_DATA
+
+public:
+
+	/** Bone name to draw widget for */
+	UPROPERTY(EditAnywhere, Category = "Settings")
+	FName BoneName = NAME_None;
+
+	/** Widget Radius */
+	UPROPERTY(EditAnywhere, Category = "Settings", meta = (ForceForceUnits = "cm"))
+	float Radius = 25.0f;
+
+	/** Widget Thickness */
+	UPROPERTY(EditAnywhere, Category = "Settings", meta = (ForceForceUnits = "cm"))
+	float Thickness = 1.0f;
+
+	/** If to depth test while drawing */
+	UPROPERTY(EditAnywhere, Category = "Settings")
+	bool bDepthTest = true;
+
+#endif
+
+#if WITH_EDITOR
+
+public:
+
+	virtual void DrawDebug_Implementation(
+		const FDebugDrawer& Drawer,
+		const FAnimDatabasePoseState& InPoseState,
+		const UAnimDatabase* InDatabase,
+		const FAnimDatabaseFrameRanges& InFrameRanges,
+		const FTransform& RangeViewportTransform,
+		const int32 SequenceIdx,
+		const float SequenceTime,
+		const int32 RangeStart,
+		const int32 RangeLength,
+		const FLinearColor& IdentifierColor) override;
+
+#endif
+};
+
 /** Debug Draw class for drawing root velocities */
 UCLASS(BlueprintType, Blueprintable, meta = (DisplayName = "Root Velocity"))
 class UAnimDatabaseDebugDraw_RootVelocity : public UAnimDatabaseDebugDraw
diff --git a/Engine/Plugins/Experimental/Animation/AnimDatabase/Source/AnimDatabase/Public/AnimDatabaseFrameAttribute.h b/Engine/Plugins/Experimental/Animation/AnimDatabase/Source/AnimDatabase/Public/AnimDatabaseFrameAttribute.h
--- a/Engine/Plugins/Experimental/Animation/AnimDatabase/Source/AnimDatabase/Public/AnimDatabaseFrameAttribute.h
+++ b/Engine/Plugins/Experimental/Animation/AnimDatabase/Source/AnimDatabase/Public/AnimDatabaseFrameAttribute.h
@@ -175,6 +175,21 @@ public:
 	UFUNCTION(BlueprintPure = true, Category = "AnimDatabase", meta = (NotBlueprintThreadSafe))
 	static UE_API FAnimDatabaseFrameAttribute MakeFrameAttributeFromFunction(const UAnimDatabase* Database, const FAnimDatabaseFrameRanges& FrameRanges, const UAnimDatabaseFrameAttributeFunction* Function);
 
+	/** Makes Frame Attributes using the given UAnimDatabaseFrameAttributeFunctions */
+	UFUNCTION(BlueprintPure = true, Category = "AnimDatabase", meta = (NotBlueprintThreadSafe))
+	static UE_API void MakeFrameAttributesFromFunctions(
+		TArray<FAnimDatabaseFrameAttribute>& OutFrameAttributes,
+		const UAnimDatabase* Database, 
+		const FAnimDatabaseFrameRanges& FrameRanges, 
+		const TArray<UAnimDatabaseFrameAttributeFunction*>& Functions);
+
+	static UE_API void MakeFrameAttributesFromFunctionsArrayView(
+		const TArrayView<FAnimDatabaseFrameAttribute> OutFrameAttributes,
+		const UAnimDatabase* Database,
+		const FAnimDatabaseFrameRanges& FrameRanges,
+		const TArrayView<UAnimDatabaseFrameAttributeFunction* const> Functions);
+
+
 	/** Make a Frame Attribute using the given UAnimDatabaseFrameAttributeFunction */
 	UFUNCTION(BlueprintPure = true, Category = "AnimDatabase", meta = (NotBlueprintThreadSafe))
 	static UE_API FAnimDatabaseFrameAttribute MakeFrameAttributeFromClass(const UAnimDatabase* Database, const FAnimDatabaseFrameRanges& FrameRanges, const TSubclassOf<UAnimDatabaseFrameAttributeFunction> Class);
@@ -315,6 +330,10 @@ public:
 	UFUNCTION(BlueprintPure = true, Category = "AnimDatabase")
 	static UE_API FAnimDatabaseFrameAttribute MakeRotationFrameAttributeFromConstantQuat(const FAnimDatabaseFrameRanges& FrameRanges, const FQuat Rotation);
 
+	/** Makes a rotation frame attribute from a set of basis directions. */
+	UFUNCTION(BlueprintPure = true, Category = "AnimDatabase")
+	static UE_API FAnimDatabaseFrameAttribute MakeRotationFrameAttributeFromBasisDirections(const FAnimDatabaseFrameAttribute& ForwardDirection, const FAnimDatabaseFrameAttribute& RightDirection, const FAnimDatabaseFrameAttribute& UpDirection);
+
 	/** Makes a scale frame attribute from the given ranges, filled with the provided value. */
 	UFUNCTION(BlueprintPure = true, Category = "AnimDatabase")
 	static UE_API FAnimDatabaseFrameAttribute MakeScaleFrameAttributeFromConstant(const FAnimDatabaseFrameRanges& FrameRanges, const FVector Scale = FVector(1, 1, 1));
@@ -1285,6 +1304,13 @@ public:
 	UFUNCTION(BlueprintPure, Category = "AnimDatabase", meta = (CompactNodeTitle = "*"))
 	static UE_API FAnimDatabaseFrameAttribute FrameAttributeTransformApply(const FTransform& A, const FAnimDatabaseFrameAttribute& B);
 
+	/** Applies a transformation to a constant location vector. Assumes types are compatible */
+	UFUNCTION(BlueprintPure, Category = "AnimDatabase")
+	static UE_API FAnimDatabaseFrameAttribute FrameAttributeApplyTransformLocation(const FAnimDatabaseFrameAttribute& A, const FVector& B);
+
+	/** Applies a transformation to a constant direction vector. Assumes types are compatible */
+	UFUNCTION(BlueprintPure, Category = "AnimDatabase")
+	static UE_API FAnimDatabaseFrameAttribute FrameAttributeApplyTransformDirection(const FAnimDatabaseFrameAttribute& A, const FVector& B);
 
 	/** Adds a constant angle value in degrees to an angle frame attribute. */
 	UFUNCTION(BlueprintPure, Category = "AnimDatabase", meta = (CompactNodeTitle = "+"))
@@ -1326,10 +1352,17 @@ public:
 
 public:
 
+	/** Computes the sum of multiple frame attributes. Assumes all are the same type */
 	UFUNCTION(BlueprintPure, Category = "AnimDatabase")
 	static UE_API FAnimDatabaseFrameAttribute FrameAttributeSum(const TArray<FAnimDatabaseFrameAttribute>& FrameAttributes);
 	static UE_API FAnimDatabaseFrameAttribute FrameAttributeSumFromArrayView(const TArrayView<const FAnimDatabaseFrameAttribute> FrameAttributes);
 
+	/** Computes the weighted sum of multiple frame attributes. Assumes all are the same type */
+	UFUNCTION(BlueprintPure, Category = "AnimDatabase")
+	static UE_API FAnimDatabaseFrameAttribute FrameAttributeWeightedSum(const TArray<FAnimDatabaseFrameAttribute>& FrameAttributes, const TArray<float>& Weights);
+	static UE_API FAnimDatabaseFrameAttribute FrameAttributeWeightedSumFromArrayView(const TArrayView<const FAnimDatabaseFrameAttribute> FrameAttributes, const TArrayView<const float> Weights);
+
+	/** Computes the product of multiple frame attributes. Assumes all are the same type */
 	UFUNCTION(BlueprintPure, Category = "AnimDatabase")
 	static UE_API FAnimDatabaseFrameAttribute FrameAttributeProduct(const TArray<FAnimDatabaseFrameAttribute>& FrameAttributes);
 	static UE_API FAnimDatabaseFrameAttribute FrameAttributeProductFromArrayView(const TArrayView<const FAnimDatabaseFrameAttribute> FrameAttributes);
@@ -1567,6 +1600,14 @@ public:
 	UFUNCTION(BlueprintPure, Category = "AnimDatabase")
 	static UE_API FAnimDatabaseFrameAttribute FrameAttributeFloatToAngleRadians(const FAnimDatabaseFrameAttribute& FloatFrameAttribute);
 
+	/** Converts a direction frame attribute to a location frame attribute */
+	UFUNCTION(BlueprintPure, Category = "AnimDatabase")
+	static UE_API FAnimDatabaseFrameAttribute FrameAttributeDirectionToLocation(const FAnimDatabaseFrameAttribute& DirectionFrameAttribute);
+
+	/** Converts a location frame attribute to a direction frame attribute (normalizing the result) */
+	UFUNCTION(BlueprintPure, Category = "AnimDatabase")
+	static UE_API FAnimDatabaseFrameAttribute FrameAttributeLocationToDirection(const FAnimDatabaseFrameAttribute& LocationFrameAttribute);
+
 public:
 
 	/** Applies a Gaussian smoothing filter to the values in each range of a frame attribute. Assumes type is compatible */
@@ -1806,6 +1847,45 @@ public:
 		const float HeightThreshold = 20.0f,
 		const float VelocityThreshold = 25.0f,
 		const int32 MajorityVoteFilterWidth = 5);
+
+	/**
+	 * Fits a transform from the given location and direction attributes
+	 *
+	 * @param Database							Database to use
+	 * @param FrameRanges						Ranges in the database to use
+	 * @param FitLocations						Array of attributes for computing the locations to use for the fit
+	 * @param FitForwardDirections				Array of attributes for computing the forward directions to use for the fit
+	 * @param FitRightDirections				Array of attributes for computing the right directions to use for the fit
+	 * @param FitUpDirections					Array of attributes for computing the up directions to use for the fit
+	 * @param FitLocationWeights				Optional array of weights to use for the fit locations
+	 * @param FitForwardDirectionWeights		Optional array of weights to use for the fit forward directions
+	 * @param FitRightDirectionWeights			Optional array of weights to use for the fit right directions
+	 * @param FitUpDirectionWeights				Optional array of weights to use for the fit up directions
+	 */
+	UFUNCTION(BlueprintPure, Category = "AnimDatabase", meta = (NotBlueprintThreadSafe))
+	static UE_API FAnimDatabaseFrameAttribute FrameAttributeFitTransform(
+		const UAnimDatabase* Database,
+		const FAnimDatabaseFrameRanges& FrameRanges,
+		const TArray<FAnimDatabaseFrameAttribute>& FitLocations,
+		const TArray<FAnimDatabaseFrameAttribute>& FitForwardDirections,
+		const TArray<FAnimDatabaseFrameAttribute>& FitRightDirections,
+		const TArray<FAnimDatabaseFrameAttribute>& FitUpDirections,
+		const TArray<float>& FitLocationWeights,
+		const TArray<float>& FitForwardDirectionWeights,
+		const TArray<float>& FitRightDirectionWeights,
+		const TArray<float>& FitUpDirectionWeights);
+
+	static UE_API FAnimDatabaseFrameAttribute FrameAttributeFitTransformFromArrayViews(
+		const UAnimDatabase* Database,
+		const FAnimDatabaseFrameRanges& FrameRanges,
+		const TArrayView<const FAnimDatabaseFrameAttribute> FitLocations,
+		const TArrayView<const FAnimDatabaseFrameAttribute> FitForwardDirections,
+		const TArrayView<const FAnimDatabaseFrameAttribute> FitRightDirections,
+		const TArrayView<const FAnimDatabaseFrameAttribute> FitUpDirections,
+		const TArrayView<const float> FitLocationWeights,
+		const TArrayView<const float> FitForwardDirectionWeights,
+		const TArrayView<const float> FitRightDirectionWeights,
+		const TArrayView<const float> FitUpDirectionWeights);
 };
 
 /** A simple class that has a single implementable function which constructs a FAnimDatabaseFrameAttribute object from a UAnimDatabase. */
@@ -1892,7 +1972,7 @@ public:
 
 	/**
 	 * Filter width (in frames) for the majority vote filter used to remove very short contact transitions.
-	 *	contact periods (or non-contact periods) less than half this duration will be removed.
+	 *Contact periods (or non-contact periods) less than half this duration will be removed.
 	 */
 	UPROPERTY(EditAnywhere, Category = "Settings")
 	int32 MajorityVoteFilterWidth = 5;
@@ -1906,6 +1986,115 @@ public:
 	UE_API virtual FAnimDatabaseFrameAttribute MakeFrameAttribute_Implementation(const UAnimDatabase* Database, const FAnimDatabaseFrameRanges& FrameRanges) const override;
 };
 
+
+/** Fits a transform to the given locations and directions */
+UCLASS(BlueprintType, Blueprintable, meta = (DisplayName = "Fit Transform"))
+class UAnimDatabaseFrameAttributeFunction_FitTransform : public UAnimDatabaseFrameAttributeFunction
+{
+	GENERATED_BODY()
+
+public:
+
+	/** Array of attributes for computing the locations to use for the fit */
+	UPROPERTY(EditAnywhere, Instanced, Category = "Settings")
+	TArray<TObjectPtr<UAnimDatabaseFrameAttributeFunction>> FitLocations;
+
+	/** Array of attributes for computing the forward directions to use for the fit */
+	UPROPERTY(EditAnywhere, Instanced, Category = "Settings")
+	TArray<TObjectPtr<UAnimDatabaseFrameAttributeFunction>> FitForwardDirections;
+
+	/** Array of attributes for computing the right directions to use for the fit */
+	UPROPERTY(EditAnywhere, Instanced, Category = "Settings")
+	TArray<TObjectPtr<UAnimDatabaseFrameAttributeFunction>> FitRightDirections;
+
+	/** Array of attributes for computing the up directions to use for the fit */
+	UPROPERTY(EditAnywhere, Instanced, Category = "Settings")
+	TArray<TObjectPtr<UAnimDatabaseFrameAttributeFunction>> FitUpDirections;
+
+	/** Optional array of weights to use for the fit locations */
+	UPROPERTY(EditAnywhere, Category = "Settings")
+	TArray<float> FitLocationWeights;
+
+	/** Optional array of weights to use for the fit forward directions */
+	UPROPERTY(EditAnywhere, Category = "Settings")
+	TArray<float> FitForwardDirectionWeights;
+
+	/** Optional array of weights to use for the fit right directions */
+	UPROPERTY(EditAnywhere, Category = "Settings")
+	TArray<float> FitRightDirectionWeights;
+
+	/** Optional array of weights to use for the fit up directions */
+	UPROPERTY(EditAnywhere, Category = "Settings")
+	TArray<float> FitUpDirectionWeights;
+
+public:
+
+	UE_API virtual FAnimDatabaseFrameAttribute MakeFrameAttribute_Implementation(const UAnimDatabase* Database, const FAnimDatabaseFrameRanges& FrameRanges) const override;
+};
+
+/** Computes a bone global transform attribute */
+UCLASS(BlueprintType, Blueprintable, meta = (DisplayName = "Bone Global Transform"))
+class UAnimDatabaseFrameAttributeFunction_BoneGlobalTransform : public UAnimDatabaseFrameAttributeFunction
+{
+	GENERATED_BODY()
+
+public:
+
+	/** Name of the bone to compute the transform of */
+	UPROPERTY(EditAnywhere, Category = "Settings")
+	FName BoneName = NAME_None;
+
+	/** Local offset transform */
+	UPROPERTY(EditAnywhere, Category = "Settings")
+	FTransform LocalOffset = FTransform::Identity;
+
+public:
+
+	UE_API virtual FAnimDatabaseFrameAttribute MakeFrameAttribute_Implementation(const UAnimDatabase* Database, const FAnimDatabaseFrameRanges& FrameRanges) const override;
+};
+
+/** Computes a bone global location attribute */
+UCLASS(BlueprintType, Blueprintable, meta = (DisplayName = "Bone Global Location"))
+class UAnimDatabaseFrameAttributeFunction_BoneGlobalLocation : public UAnimDatabaseFrameAttributeFunction
+{
+	GENERATED_BODY()
+
+public:
+
+	/** Name of the bone to compute the location of */
+	UPROPERTY(EditAnywhere, Category = "Settings")
+	FName BoneName = NAME_None;
+
+	/** Local offset vector */
+	UPROPERTY(EditAnywhere, Category = "Settings")
+	FVector LocalOffset = FVector::ZeroVector;
+
+public:
+
+	UE_API virtual FAnimDatabaseFrameAttribute MakeFrameAttribute_Implementation(const UAnimDatabase* Database, const FAnimDatabaseFrameRanges& FrameRanges) const override;
+};
+
+/** Computes a bone global direction attribute */
+UCLASS(BlueprintType, Blueprintable, meta = (DisplayName = "Bone Global Direction"))
+class UAnimDatabaseFrameAttributeFunction_BoneGlobalDirection : public UAnimDatabaseFrameAttributeFunction
+{
+	GENERATED_BODY()
+
+public:
+
+	/** Name of the bone to compute the direction of */
+	UPROPERTY(EditAnywhere, Category = "Settings")
+	FName BoneName = NAME_None;
+
+	/** Local direction vector */
+	UPROPERTY(EditAnywhere, Category = "Settings")
+	FVector LocalDirection = FVector::ForwardVector;
+
+public:
+
+	UE_API virtual FAnimDatabaseFrameAttribute MakeFrameAttribute_Implementation(const UAnimDatabase* Database, const FAnimDatabaseFrameRanges& FrameRanges) const override;
+};
+
 /** Computes a bone velocity magnitude attribute */
 UCLASS(BlueprintType, Blueprintable, meta = (DisplayName = "Bone Global Linear Velocity Magnitude"))
 class UAnimDatabaseFrameAttributeFunction_BoneGlobalLinearVelocityMagnitude : public UAnimDatabaseFrameAttributeFunction
diff --git a/Engine/Plugins/Experimental/Animation/AnimDatabase/Source/AnimDatabase/Public/AnimDatabaseMath.h b/Engine/Plugins/Experimental/Animation/AnimDatabase/Source/AnimDatabase/Public/AnimDatabaseMath.h
--- a/Engine/Plugins/Experimental/Animation/AnimDatabase/Source/AnimDatabase/Public/AnimDatabaseMath.h
+++ b/Engine/Plugins/Experimental/Animation/AnimDatabase/Source/AnimDatabase/Public/AnimDatabaseMath.h
@@ -1010,6 +1010,8 @@ namespace UE::AnimDatabase::Math
 		const float FrameDeltaTime);
 
 	// Mean/Std/Scale/Min/Max
+	
+	UE_API float Sum(const TLearningArrayView<1, const float> Data);
 
 	UE_API void ComputeMean(
 		float& OutMean,
