AnimDatabase: Various improvements, fixes, and functions to better allow for prop fitting.

#Synced-CL 51630663

diff --git a/Engine/Plugins/Experimental/Animation/AnimDatabase/Source/AnimDatabase/Private/AnimDatabase.cpp b/Engine/Plugins/Experimental/Animation/AnimDatabase/Source/AnimDatabase/Private/AnimDatabase.cpp
--- a/Engine/Plugins/Experimental/Animation/AnimDatabase/Source/AnimDatabase/Private/AnimDatabase.cpp
+++ b/Engine/Plugins/Experimental/Animation/AnimDatabase/Source/AnimDatabase/Private/AnimDatabase.cpp
@@ -6321,8 +6321,14 @@ void UAnimDatabaseLibrary::DatabaseSetLocalBoneTransformsFromArrayViews(UAnimDat
 		Controllers, SequenceControllerIndices, Database, TransformFrameAttributes,
 		LOCTEXT("SetBoneTransformScript", "Setting Bone Transform"), bShouldTransact);
 
+	FScopedSlowTask SlowTask(BoneNum, LOCTEXT("SettingBoneLocalTransforms", "Setting Bone Transforms"));
+	SlowTask.MakeDialog(true);
+
 	for (int32 BoneIdx = 0; BoneIdx < BoneNum; BoneIdx++)
 	{
+		if (SlowTask.ShouldCancel()) { break; }
+		SlowTask.EnterProgressFrame();
+
 		const FAnimDatabaseFrameAttribute FilteredFrameAttribute = UAnimDatabaseFrameAttributeLibrary::FrameAttributeIntersection(TransformFrameAttributes[BoneIdx],
 			UAnimDatabaseFrameRangesLibrary::MakeFrameRangesFromNotMirrored(Database, UAnimDatabaseFrameAttributeLibrary::FrameAttributeFrameRanges(TransformFrameAttributes[BoneIdx])));
 
@@ -6406,6 +6412,48 @@ void UAnimDatabaseLibrary::DatabaseSetLocalBoneTransformsFromArrayViews(UAnimDat
 #endif
 }
 
+void UAnimDatabaseLibrary::DatabaseSetGlobalBoneTransform(UAnimDatabase* Database, const FName BoneName, const FAnimDatabaseFrameAttribute& TransformFrameAttribute, const bool bShouldTransact)
+{
+	DatabaseSetGlobalBoneTransformsFromArrayViews(Database, { BoneName }, { TransformFrameAttribute }, bShouldTransact);
+}
+
+void UAnimDatabaseLibrary::DatabaseSetGlobalBoneTransforms(UAnimDatabase* Database, const TArray<FName>& BoneNames, const TArray<FAnimDatabaseFrameAttribute>& TransformFrameAttributes, const bool bShouldTransact)
+{
+	DatabaseSetGlobalBoneTransformsFromArrayViews(Database, BoneNames, TransformFrameAttributes, bShouldTransact);
+}
+
+void UAnimDatabaseLibrary::DatabaseSetGlobalBoneTransformsFromArrayViews(UAnimDatabase* Database, const TArrayView<const FName> BoneNames, const TArrayView<const FAnimDatabaseFrameAttribute> TransformFrameAttributes, const bool bShouldTransact)
+{
+#if WITH_EDITOR
+
+	const int32 BoneNum = BoneNames.Num();
+
+	if (BoneNames.Num() != TransformFrameAttributes.Num())
+	{
+		return;
+	}
+
+	TArray<FAnimDatabaseFrameAttribute> LocalTransforms;
+	LocalTransforms.SetNum(BoneNum);
+
+	for (int32 BoneIdx = 0; BoneIdx < BoneNum; BoneIdx++)
+	{
+		const int32 ParentBoneIdx = Database->GetBoneParent(Database->FindBoneIndex(BoneNames[BoneIdx]));
+		if (ParentBoneIdx != INDEX_NONE)
+		{
+			LocalTransforms[BoneIdx] = UAnimDatabaseFrameAttributeLibrary::FrameAttributeDivide(TransformFrameAttributes[BoneIdx],
+				UAnimDatabaseFrameAttributeLibrary::MakeBoneGlobalTransformFrameAttribute(Database, UAnimDatabaseFrameAttributeLibrary::FrameAttributeFrameRanges(TransformFrameAttributes[BoneIdx]), ParentBoneIdx));
+		}
+		else
+		{
+			LocalTransforms[BoneIdx] = TransformFrameAttributes[BoneIdx];
+		}
+	}
+
+	DatabaseSetLocalBoneTransformsFromArrayViews(Database, BoneNames, LocalTransforms, bShouldTransact);
+#endif
+}
+
 void UAnimDatabaseLibrary::DatabaseStripUnusedBoneTracks(UAnimDatabase* Database, const FAnimDatabaseFrameRanges& FrameRanges, const bool bShouldTransact)
 {
 #if WITH_EDITOR
@@ -6426,7 +6474,7 @@ void UAnimDatabaseLibrary::DatabaseStripUnusedBoneTracks(UAnimDatabase* Database
 		const int32 EntryIdx,
 		const int32 RangeIdx) {
 
-			TRACE_CPUPROFILER_EVENT_SCOPE(UAnimDatabaseLibrary::DatabaseModifyPoseData::SampleAndModify);
+			TRACE_CPUPROFILER_EVENT_SCOPE(UAnimDatabaseLibrary::DatabaseStripUnusedBoneTracks::Sample);
 
 			const int32 FrameOffset = FilteredFrameRanges.FrameRangeSet->GetEntryRangeOffset(EntryIdx, RangeIdx);
 			const int32 FrameNum = FilteredFrameRanges.FrameRangeSet->GetEntryRangeLength(EntryIdx, RangeIdx);
@@ -6508,6 +6556,14 @@ void UAnimDatabaseLibrary::DatabaseModifyPoseData(UAnimDatabase* Database, const
 			Function(PoseData.Slice(FrameOffset, FrameNum));
 		});
 
+	// Compute the required used bones
+
+	TArray<int32> UsedBones;
+	UE::AnimDatabase::PoseData::ComputeUsedBones(
+		UsedBones,
+		PoseData.LocalBoneData.ConstView(),
+		Database->Skeleton->GetReferenceSkeleton().GetRefBonePose());
+
 	// Write back modified pose data
 
 	const int32 TotalRangeNum = FilteredFrameRanges.FrameRangeSet->GetTotalRangeNum();
@@ -6523,28 +6579,26 @@ void UAnimDatabaseLibrary::DatabaseModifyPoseData(UAnimDatabase* Database, const
 	TagRangeSequenceIndices.SetNumUninitialized({ TotalRangeNum });
 	UE::Learning::FrameRangeSet::AllRangeSequences(TagRangeSequenceIndices, *FilteredFrameRanges.FrameRangeSet);
 
+	FScopedSlowTask SlowTask(TotalRangeNum, LOCTEXT("ModifyingRanges", "Modifying Sequence Ranges"));
+	SlowTask.MakeDialog(true);
+
 	for (int32 RangeIdx = 0; RangeIdx < TotalRangeNum; RangeIdx++)
 	{
 		TRACE_CPUPROFILER_EVENT_SCOPE(UAnimDatabaseLibrary::DatabaseModifyPoseData::WriteBack);
 
+		if (SlowTask.ShouldCancel()) { break; }
+		SlowTask.EnterProgressFrame();
+
 		const int32 RangeSequenceIdx = TagRangeSequenceIndices[RangeIdx];
 		check(!Database->GetIsMirrored(RangeSequenceIdx));
 
-		TScriptInterface<IAnimationDataController> Controller = Controllers[SequenceControllerIndices[RangeSequenceIdx]];
-
 		const int32 FrameOffset = FilteredFrameRanges.FrameRangeSet->GetAllRangeOffsets()[RangeIdx];
 		const int32 FrameNum = FilteredFrameRanges.FrameRangeSet->GetAllRangeLengths()[RangeIdx];
 		const int32 StartFrame = FilteredFrameRanges.FrameRangeSet->GetAllRangeStarts()[RangeIdx];
 
-		TArray<int32> UsedBones;
-		UE::AnimDatabase::PoseData::ComputeUsedBones(
-			UsedBones, 
-			PoseData.LocalBoneData.ConstSlice(FrameOffset, FrameNum), 
-			Database->Skeleton->GetReferenceSkeleton().GetRefBonePose());
-
 		SetAnimSequencePoseData(
 			Database->GetAnimSequence(RangeSequenceIdx),
-			Controller,
+			Controllers[SequenceControllerIndices[RangeSequenceIdx]],
 			StartFrame,
 			PoseData.ConstSlice(FrameOffset, FrameNum),
 			UsedBones,
@@ -7245,6 +7299,47 @@ void UAnimDatabaseDebugDraw_BoneTransform::DrawDebug_Implementation(
 		Radius);
 }
 
+
+void UAnimDatabaseDebugDraw_BoneAttachment::DrawDebug_Implementation(
+	const FDebugDrawer& Drawer,
+	const FAnimDatabasePoseState& InPoseState,
+	const UAnimDatabase* InDatabase,
+	const FAnimDatabaseFrameRanges& InFrameRanges,
+	const FTransform& RangeViewportTransform,
+	const int32 SequenceIdx,
+	const float SequenceTime,
+	const int32 RangeStart,
+	const int32 RangeLength,
+	const FLinearColor& IdentifierColor)
+{
+	const int32 BoneIndex = InDatabase->FindBoneIndex(BoneName);
+	if (BoneIndex == INDEX_NONE) { return; }
+
+	const FTransform BoneTransform = UAnimDatabasePoseStateLibrary::PoseStateBoneWorldTransform(InPoseState, BoneIndex);
+	const FTransform AttachmentTransform = RelativeTransform * BoneTransform;
+
+	if (bDrawTransform)
+	{
+		UDrawDebugLibrary::DrawDebugTransform(
+			Drawer,
+			AttachmentTransform,
+			UDrawDebugLibrary::MakeDrawDebugLineStyleFromColorAndThickness(IdentifierColor, TransformThickness),
+			bDepthTest,
+			TransformRadius);
+	}
+
+	if (StaticMesh)
+	{
+		UDrawDebugLibrary::DrawDebugStaticMeshBoundingBox(
+			Drawer, 
+			StaticMesh, 
+			AttachmentTransform.GetLocation(), 
+			AttachmentTransform.GetRotation().Rotator(),
+			UDrawDebugLibrary::MakeDrawDebugLineStyleFromColorAndThickness(IdentifierColor, BoundingBoxThickness),
+			bDepthTest);
+	}
+}
+
 void UAnimDatabaseDebugDraw_RootVelocity::InitializeDrawDebug_Implementation(const UAnimDatabase* InDatabase, const FAnimDatabaseFrameRanges& InFrameRanges)
 {
 	RootLinearVelocity = UAnimDatabaseFrameAttributeLibrary::FrameAttributeLength(UAnimDatabaseFrameAttributeLibrary::MakeRootLinearVelocityFrameAttribute(InDatabase, InFrameRanges));
diff --git a/Engine/Plugins/Experimental/Animation/AnimDatabase/Source/AnimDatabase/Private/AnimDatabaseFrameAttribute.cpp b/Engine/Plugins/Experimental/Animation/AnimDatabase/Source/AnimDatabase/Private/AnimDatabaseFrameAttribute.cpp
--- a/Engine/Plugins/Experimental/Animation/AnimDatabase/Source/AnimDatabase/Private/AnimDatabaseFrameAttribute.cpp
+++ b/Engine/Plugins/Experimental/Animation/AnimDatabase/Source/AnimDatabase/Private/AnimDatabaseFrameAttribute.cpp
@@ -3141,6 +3141,7 @@ void UAnimDatabaseFrameAttributeLibrary::MakeBoneGlobalTransformFrameAttributesF
 	const TArray<FName>& BoneNames,
 	const float RelativeTime)
 {
+	OutTransformFrameAttributes.SetNum(BoneNames.Num());
 	MakeBoneGlobalTransformFrameAttributesFromNamesArrayView(
 		OutTransformFrameAttributes,
 		Database,
@@ -3150,7 +3151,7 @@ void UAnimDatabaseFrameAttributeLibrary::MakeBoneGlobalTransformFrameAttributesF
 }
 
 void UAnimDatabaseFrameAttributeLibrary::MakeBoneGlobalTransformFrameAttributesFromNamesArrayView(
-	TArray<FAnimDatabaseFrameAttribute>& OutTransformFrameAttributes,
+	TArrayView<FAnimDatabaseFrameAttribute> OutTransformFrameAttributes,
 	const UAnimDatabase* Database,
 	const FAnimDatabaseFrameRanges& FrameRanges,
 	const TArrayView<const FName> BoneNames,
@@ -3159,14 +3160,12 @@ void UAnimDatabaseFrameAttributeLibrary::MakeBoneGlobalTransformFrameAttributesF
 	if (!Database)
 	{
 		UE_LOGFMT(LogAnimDatabase, Error, "MakeBoneGlobalTransformFrameAttributesFromNamesArrayView: Database is nullptr.");
-		OutTransformFrameAttributes.Empty();
 		return;
 	}
 
 	TArray<int32, TInlineAllocator<32>> BoneIndices;
 	const int32 BoneNum = BoneNames.Num();
 	BoneIndices.SetNumUninitialized(BoneNum);
-	OutTransformFrameAttributes.SetNum(BoneNum);
 	Database->FindBoneIndicesFromArrayViews(BoneIndices, BoneNames);
 
 	MakeBoneGlobalFrameAttributesFromArrayViews(
@@ -3211,6 +3210,51 @@ FAnimDatabaseFrameAttribute UAnimDatabaseFrameAttributeLibrary::MakeBoneGlobalLo
 	return MakeBoneGlobalLocationFrameAttribute(Database, FrameRanges, BoneIdx, RelativeTime);
 }
 
+void UAnimDatabaseFrameAttributeLibrary::MakeBoneGlobalLocationFrameAttributes(TArray<FAnimDatabaseFrameAttribute>& OutLocationFrameAttributes, const UAnimDatabase* Database, const FAnimDatabaseFrameRanges& FrameRanges, const TArray<int32>& BoneIndices, const float RelativeTime)
+{
+	OutLocationFrameAttributes.SetNum(BoneIndices.Num());
+	MakeBoneGlobalLocationFrameAttributesFromArrayViews(OutLocationFrameAttributes, Database, FrameRanges, BoneIndices, RelativeTime);
+}
+
+void UAnimDatabaseFrameAttributeLibrary::MakeBoneGlobalLocationFrameAttributesFromArrayViews(TArrayView<FAnimDatabaseFrameAttribute> OutLocationFrameAttributes, const UAnimDatabase* Database, const FAnimDatabaseFrameRanges& FrameRanges, const TArrayView<const int32> BoneIndices, const float RelativeTime)
+{
+	MakeBoneGlobalFrameAttributesFromArrayViews(
+		OutLocationFrameAttributes, {}, {},
+		{}, {}, {},
+		Database,
+		FrameRanges,
+		BoneIndices, {}, {},
+		{}, {}, {},
+		RelativeTime);
+}
+
+void UAnimDatabaseFrameAttributeLibrary::MakeBoneGlobalLocationFrameAttributesFromNames(TArray<FAnimDatabaseFrameAttribute>& OutLocationFrameAttributes, const UAnimDatabase* Database, const FAnimDatabaseFrameRanges& FrameRanges, const TArray<FName>& BoneNames, const float RelativeTime)
+{
+	OutLocationFrameAttributes.SetNum(BoneNames.Num());
+	MakeBoneGlobalLocationFrameAttributesFromNamesArrayView(OutLocationFrameAttributes, Database, FrameRanges, BoneNames, RelativeTime);
+}
+
+void UAnimDatabaseFrameAttributeLibrary::MakeBoneGlobalLocationFrameAttributesFromNamesArrayView(TArrayView<FAnimDatabaseFrameAttribute> OutLocationFrameAttributes, const UAnimDatabase* Database, const FAnimDatabaseFrameRanges& FrameRanges, const TArrayView<const FName> BoneNames, const float RelativeTime)
+{
+	if (!Database)
+	{
+		UE_LOGFMT(LogAnimDatabase, Error, "MakeBoneGlobalLocationFrameAttributesFromNamesArrayView: Database is nullptr.");
+		return;
+	}
+
+	TArray<int32, TInlineAllocator<32>> BoneIndices;
+	const int32 BoneNum = BoneNames.Num();
+	BoneIndices.SetNumUninitialized(BoneNum);
+	Database->FindBoneIndicesFromArrayViews(BoneIndices, BoneNames);
+
+	MakeBoneGlobalLocationFrameAttributesFromArrayViews(
+		OutLocationFrameAttributes,
+		Database,
+		FrameRanges,
+		BoneIndices,
+		RelativeTime);
+}
+
 FAnimDatabaseFrameAttribute UAnimDatabaseFrameAttributeLibrary::MakeBoneGlobalDirectionFrameAttribute(const UAnimDatabase* Database, const FAnimDatabaseFrameRanges& FrameRanges, const int32 BoneIndex, const FVector LocalDirection, const float RelativeTime)
 {
 	FAnimDatabaseFrameAttribute OutFrameAttribute;
@@ -3245,6 +3289,67 @@ FAnimDatabaseFrameAttribute UAnimDatabaseFrameAttributeLibrary::MakeBoneGlobalDi
 	return MakeBoneGlobalDirectionFrameAttribute(Database, FrameRanges, BoneIdx, LocalDirection, RelativeTime);
 }
 
+void UAnimDatabaseFrameAttributeLibrary::MakeBoneGlobalDirectionFrameAttributes(TArray<FAnimDatabaseFrameAttribute>& OutDirectionFrameAttributes, const UAnimDatabase* Database, const FAnimDatabaseFrameRanges& FrameRanges, const TArray<int32>& BoneIndices, const FVector LocalDirection, const float RelativeTime)
+{
+	OutDirectionFrameAttributes.SetNum(BoneIndices.Num());
+	MakeBoneGlobalDirectionFrameAttributesFromArrayViews(OutDirectionFrameAttributes, Database, FrameRanges, BoneIndices, LocalDirection, RelativeTime);
+}
+
+void UAnimDatabaseFrameAttributeLibrary::MakeBoneGlobalDirectionFrameAttributesFromArrayViews(TArrayView<FAnimDatabaseFrameAttribute> OutDirectionFrameAttributes, const UAnimDatabase* Database, const FAnimDatabaseFrameRanges& FrameRanges, const TArrayView<const int32> BoneIndices, const FVector LocalDirection, const float RelativeTime)
+{
+	if (BoneIndices.Num() != OutDirectionFrameAttributes.Num())
+	{
+		return;
+	}
+
+	TArray<FAnimDatabaseFrameAttribute, TInlineAllocator<32>> RotationAttributes;
+	RotationAttributes.SetNum(BoneIndices.Num());
+
+	MakeBoneGlobalFrameAttributesFromArrayViews(
+		{}, RotationAttributes, {},
+		{}, {}, {},
+		Database,
+		FrameRanges,
+		{}, BoneIndices, {},
+		{}, {}, {},
+		RelativeTime);
+
+	const int32 AttrNum = OutDirectionFrameAttributes.Num();
+	for (int32 AttrIdx = 0; AttrIdx < AttrNum; AttrIdx++)
+	{
+		OutDirectionFrameAttributes[AttrIdx] = FrameAttributeRotateDirection(RotationAttributes[AttrIdx], LocalDirection);
+	}
+}
+
+void UAnimDatabaseFrameAttributeLibrary::MakeBoneGlobalDirectionFrameAttributesFromNames(TArray<FAnimDatabaseFrameAttribute>& OutDirectionFrameAttributes, const UAnimDatabase* Database, const FAnimDatabaseFrameRanges& FrameRanges, const TArray<FName>& BoneNames, const FVector LocalDirection, const float RelativeTime)
+{
+	OutDirectionFrameAttributes.SetNum(BoneNames.Num());
+	MakeBoneGlobalDirectionFrameAttributesFromNamesArrayView(OutDirectionFrameAttributes, Database, FrameRanges, BoneNames, LocalDirection, RelativeTime);
+}
+
+void UAnimDatabaseFrameAttributeLibrary::MakeBoneGlobalDirectionFrameAttributesFromNamesArrayView(TArrayView<FAnimDatabaseFrameAttribute> OutDirectionFrameAttributes, const UAnimDatabase* Database, const FAnimDatabaseFrameRanges& FrameRanges, const TArrayView<const FName> BoneNames, const FVector LocalDirection, const float RelativeTime)
+{
+	if (!Database)
+	{
+		UE_LOGFMT(LogAnimDatabase, Error, "MakeBoneGlobalDirectionFrameAttributesFromNamesArrayView: Database is nullptr.");
+		return;
+	}
+
+	TArray<int32, TInlineAllocator<32>> BoneIndices;
+	const int32 BoneNum = BoneNames.Num();
+	BoneIndices.SetNumUninitialized(BoneNum);
+	Database->FindBoneIndicesFromArrayViews(BoneIndices, BoneNames);
+
+	MakeBoneGlobalDirectionFrameAttributesFromArrayViews(
+		OutDirectionFrameAttributes,
+		Database,
+		FrameRanges,
+		BoneIndices,
+		LocalDirection,
+		RelativeTime);
+}
+
+
 FAnimDatabaseFrameAttribute UAnimDatabaseFrameAttributeLibrary::MakeBoneGlobalLinearVelocityFrameAttribute(const UAnimDatabase* Database, const FAnimDatabaseFrameRanges& FrameRanges, const int32 BoneIndex, const float RelativeTime)
 {
 	FAnimDatabaseFrameAttribute OutFrameAttribute;
@@ -3347,6 +3452,39 @@ FAnimDatabaseFrameAttribute UAnimDatabaseFrameAttributeLibrary::MakeBoneGlobalTr
 	return MakeBoneGlobalTransformFrameAttribute(Database, FrameRanges, BoneIdx, RelativeTime);
 }
 
+FAnimDatabaseFrameAttribute UAnimDatabaseFrameAttributeLibrary::MakeGlobalDirectionBetweenBonesFrameAttribute(const UAnimDatabase* Database, const FAnimDatabaseFrameRanges& FrameRanges, const int32 BoneIndexFrom, const int32 BoneIndexTo)
+{
+	TArray<FAnimDatabaseFrameAttribute> Locations;
+	Locations.SetNum(2);
+	MakeBoneGlobalLocationFrameAttributesFromArrayViews(Locations, Database, FrameRanges, { BoneIndexFrom, BoneIndexTo });
+	return FrameAttributeLocationToDirection(FrameAttributeSubtract(Locations[1], Locations[0]));
+}
+
+FAnimDatabaseFrameAttribute UAnimDatabaseFrameAttributeLibrary::MakeGlobalDirectionBetweenBonesFrameAttributeFromNames(const UAnimDatabase* Database, const FAnimDatabaseFrameRanges& FrameRanges, const FName BoneNameFrom, const FName BoneNameTo)
+{
+	if (!Database)
+	{
+		UE_LOGFMT(LogAnimDatabase, Error, "MakeGlobalDirectionBetweenBonesFrameAttribute: Database is nullptr.");
+		return FAnimDatabaseFrameAttribute();
+	}
+
+	const int32 BoneIdxFrom = Database->FindBoneIndex(BoneNameFrom);
+	if (BoneIdxFrom == INDEX_NONE)
+	{
+		UE_LOGFMT(LogAnimDatabase, Error, "MakeGlobalDirectionBetweenBonesFrameAttribute: Bone {Name} not found.", *BoneNameFrom.ToString());
+		return FAnimDatabaseFrameAttribute();
+	}
+
+	const int32 BoneIdxTo = Database->FindBoneIndex(BoneNameTo);
+	if (BoneIdxTo == INDEX_NONE)
+	{
+		UE_LOGFMT(LogAnimDatabase, Error, "MakeGlobalDirectionBetweenBonesFrameAttribute: Bone {Name} not found.", *BoneNameTo.ToString());
+		return FAnimDatabaseFrameAttribute();
+	}
+
+	return MakeGlobalDirectionBetweenBonesFrameAttribute(Database, FrameRanges, BoneIdxFrom, BoneIdxTo);
+}
+
 void UAnimDatabaseFrameAttributeLibrary::MakePoseFrameAttribute(
 	TArray<FAnimDatabaseFrameAttribute>& OutLocationFrameAttributes,
 	TArray<FAnimDatabaseFrameAttribute>& OutLinearVelocityFrameAttributes,
@@ -6725,6 +6863,8 @@ FAnimDatabaseFrameAttribute UAnimDatabaseFrameAttributeLibrary::FrameAttributeMu
 	if (!(
 		(A.Type == EAnimDatabaseAttributeType::Float && B.Type == EAnimDatabaseAttributeType::Float) ||
 		(A.Type == EAnimDatabaseAttributeType::Location && B.Type == EAnimDatabaseAttributeType::Location) ||
+		(A.Type == EAnimDatabaseAttributeType::Location && B.Type == EAnimDatabaseAttributeType::Float) ||
+		(A.Type == EAnimDatabaseAttributeType::Float && B.Type == EAnimDatabaseAttributeType::Location) ||
 		(A.Type == EAnimDatabaseAttributeType::Rotation && B.Type == EAnimDatabaseAttributeType::Rotation) ||
 		(A.Type == EAnimDatabaseAttributeType::Transform && B.Type == EAnimDatabaseAttributeType::Transform) ||
 		(A.Type == EAnimDatabaseAttributeType::Angle && B.Type == EAnimDatabaseAttributeType::Float) ||
@@ -6742,6 +6882,24 @@ FAnimDatabaseFrameAttribute UAnimDatabaseFrameAttributeLibrary::FrameAttributeMu
 		UE::Learning::FrameAttribute::Mul(*Out.FrameAttribute, *A.FrameAttribute, *B.FrameAttribute);
 		return Out;
 	}
+	else if (A.Type == EAnimDatabaseAttributeType::Float && B.Type == EAnimDatabaseAttributeType::Location)
+	{
+		FAnimDatabaseFrameAttribute Out = MakeEmptyFrameAttribute();
+		Out.Type = B.Type;
+		UE::Learning::FFrameAttribute LocationA;
+		UE::Learning::FrameAttribute::Repeat(LocationA, 3, *A.FrameAttribute);
+		UE::Learning::FrameAttribute::Mul(*Out.FrameAttribute, LocationA, *B.FrameAttribute);
+		return Out;
+	}
+	else if (A.Type == EAnimDatabaseAttributeType::Location && B.Type == EAnimDatabaseAttributeType::Float)
+	{
+		FAnimDatabaseFrameAttribute Out = MakeEmptyFrameAttribute();
+		Out.Type = A.Type;
+		UE::Learning::FFrameAttribute LocationB;
+		UE::Learning::FrameAttribute::Repeat(LocationB, 3, *B.FrameAttribute);
+		UE::Learning::FrameAttribute::Mul(*Out.FrameAttribute, *A.FrameAttribute, LocationB);
+		return Out;
+	}
 	else if (A.Type == EAnimDatabaseAttributeType::Rotation && B.Type == EAnimDatabaseAttributeType::Rotation)
 	{
 		FAnimDatabaseFrameAttribute Out = MakeEmptyFrameAttribute();
@@ -6828,6 +6986,104 @@ FAnimDatabaseFrameAttribute UAnimDatabaseFrameAttributeLibrary::FrameAttributeDi
 	}
 }
 
+FAnimDatabaseFrameAttribute UAnimDatabaseFrameAttributeLibrary::FrameAttributeMax(FAnimDatabaseFrameAttribute A, FAnimDatabaseFrameAttribute B)
+{
+	if (!A.IsValid() || !B.IsValid())
+	{
+		UE_LOGFMT(LogAnimDatabase, Error, "FrameAttributeMax: Invalid FrameAttribute.");
+		return FAnimDatabaseFrameAttribute();
+	}
+
+	if (!(
+		(A.Type == EAnimDatabaseAttributeType::Float && B.Type == EAnimDatabaseAttributeType::Float) ||
+		(A.Type == EAnimDatabaseAttributeType::Location && B.Type == EAnimDatabaseAttributeType::Location)))
+	{
+		UE_LOGFMT(LogAnimDatabase, Error, "FrameAttributeMax: Invalid FrameAttribute Types. Got {Type0} and {Type1}.", AttributeTypeNameInternal(A.Type), AttributeTypeNameInternal(B.Type));
+		return FAnimDatabaseFrameAttribute();
+	}
+
+	EAnimDatabaseAttributeType OutputType = EAnimDatabaseAttributeType::Float;
+	if (A.Type == EAnimDatabaseAttributeType::Location && B.Type == EAnimDatabaseAttributeType::Location) { OutputType = EAnimDatabaseAttributeType::Location; }
+
+	FAnimDatabaseFrameAttribute Out = MakeEmptyFrameAttribute();
+	Out.Type = OutputType;
+	UE::Learning::FrameAttribute::Max(*Out.FrameAttribute, *A.FrameAttribute, *B.FrameAttribute);
+	return Out;
+}
+
+FAnimDatabaseFrameAttribute UAnimDatabaseFrameAttributeLibrary::FrameAttributeMin(FAnimDatabaseFrameAttribute A, FAnimDatabaseFrameAttribute B)
+{
+	if (!A.IsValid() || !B.IsValid())
+	{
+		UE_LOGFMT(LogAnimDatabase, Error, "FrameAttributeMin: Invalid FrameAttribute.");
+		return FAnimDatabaseFrameAttribute();
+	}
+
+	if (!(
+		(A.Type == EAnimDatabaseAttributeType::Float && B.Type == EAnimDatabaseAttributeType::Float) ||
+		(A.Type == EAnimDatabaseAttributeType::Location && B.Type == EAnimDatabaseAttributeType::Location)))
+	{
+		UE_LOGFMT(LogAnimDatabase, Error, "FrameAttributeMin: Invalid FrameAttribute Types. Got {Type0} and {Type1}.", AttributeTypeNameInternal(A.Type), AttributeTypeNameInternal(B.Type));
+		return FAnimDatabaseFrameAttribute();
+	}
+
+	EAnimDatabaseAttributeType OutputType = EAnimDatabaseAttributeType::Float;
+	if (A.Type == EAnimDatabaseAttributeType::Location && B.Type == EAnimDatabaseAttributeType::Location) { OutputType = EAnimDatabaseAttributeType::Location; }
+
+	FAnimDatabaseFrameAttribute Out = MakeEmptyFrameAttribute();
+	Out.Type = OutputType;
+	UE::Learning::FrameAttribute::Min(*Out.FrameAttribute, *A.FrameAttribute, *B.FrameAttribute);
+	return Out;
+}
+
+FAnimDatabaseFrameAttribute UAnimDatabaseFrameAttributeLibrary::FrameAttributeDot(FAnimDatabaseFrameAttribute A, FAnimDatabaseFrameAttribute B)
+{
+	if (!A.IsValid() || !B.IsValid())
+	{
+		UE_LOGFMT(LogAnimDatabase, Error, "FrameAttributeDot: Invalid FrameAttribute.");
+		return FAnimDatabaseFrameAttribute();
+	}
+
+	if (!(
+		(A.Type == EAnimDatabaseAttributeType::Location && B.Type == EAnimDatabaseAttributeType::Location) ||
+		(A.Type == EAnimDatabaseAttributeType::Direction && B.Type == EAnimDatabaseAttributeType::Direction)))
+	{
+		UE_LOGFMT(LogAnimDatabase, Error, "FrameAttributeDot: Invalid FrameAttribute Types. Got {Type0} and {Type1}.", AttributeTypeNameInternal(A.Type), AttributeTypeNameInternal(B.Type));
+		return FAnimDatabaseFrameAttribute();
+	}
+
+	FAnimDatabaseFrameAttribute Out = MakeEmptyFrameAttribute();
+	Out.Type = EAnimDatabaseAttributeType::Float;
+	UE::Learning::FrameAttribute::Dot(*Out.FrameAttribute, *A.FrameAttribute, *B.FrameAttribute);
+	return Out;
+}
+
+FAnimDatabaseFrameAttribute UAnimDatabaseFrameAttributeLibrary::FrameAttributeCross(FAnimDatabaseFrameAttribute A, FAnimDatabaseFrameAttribute B)
+{
+	if (!A.IsValid() || !B.IsValid())
+	{
+		UE_LOGFMT(LogAnimDatabase, Error, "FrameAttributeDot: Invalid FrameAttribute.");
+		return FAnimDatabaseFrameAttribute();
+	}
+
+	if (!(
+		(A.Type == EAnimDatabaseAttributeType::Location && B.Type == EAnimDatabaseAttributeType::Location) ||
+		(A.Type == EAnimDatabaseAttributeType::Direction && B.Type == EAnimDatabaseAttributeType::Direction)))
+	{
+		UE_LOGFMT(LogAnimDatabase, Error, "FrameAttributeDot: Invalid FrameAttribute Types. Got {Type0} and {Type1}.", AttributeTypeNameInternal(A.Type), AttributeTypeNameInternal(B.Type));
+		return FAnimDatabaseFrameAttribute();
+	}
+
+	FAnimDatabaseFrameAttribute Out = MakeEmptyFrameAttribute();
+	Out.Type = A.Type;
+	UE::Learning::FrameAttribute::Cross(*Out.FrameAttribute, *A.FrameAttribute, *B.FrameAttribute);
+	if (Out.Type == EAnimDatabaseAttributeType::Direction)
+	{
+		UE::Learning::FrameAttribute::NormalizeInplace(*Out.FrameAttribute);
+	}
+	return Out;
+}
+
 FAnimDatabaseFrameAttribute UAnimDatabaseFrameAttributeLibrary::FrameAttributeAddFloat(const FAnimDatabaseFrameAttribute& A, const float B)
 {
 	if (!A.IsValid())
@@ -6980,6 +7236,62 @@ FAnimDatabaseFrameAttribute UAnimDatabaseFrameAttributeLibrary::FrameAttributeDi
 	return Out;
 }
 
+FAnimDatabaseFrameAttribute UAnimDatabaseFrameAttributeLibrary::FrameAttributeMaxFloat(const FAnimDatabaseFrameAttribute& A, const float B)
+{
+	if (!A.IsValid())
+	{
+		UE_LOGFMT(LogAnimDatabase, Error, "FrameAttributeMaxFloat: Invalid FrameAttribute.");
+		return FAnimDatabaseFrameAttribute();
+	}
+
+	if (A.Type != EAnimDatabaseAttributeType::Float &&
+		A.Type != EAnimDatabaseAttributeType::Location)
+	{
+		UE_LOGFMT(LogAnimDatabase, Error, "FrameAttributeMaxFloat: FrameAttribute Type not valid. Got {Type}.", AttributeTypeNameInternal(A.Type));
+		return FAnimDatabaseFrameAttribute();
+	}
+
+	FAnimDatabaseFrameAttribute Out = MakeEmptyFrameAttribute();
+	Out.Type = A.Type;
+	if (A.Type == EAnimDatabaseAttributeType::Float)
+	{
+		UE::Learning::FrameAttribute::MaxConstant(*Out.FrameAttribute, *A.FrameAttribute, { B });
+	}
+	else
+	{
+		UE::Learning::FrameAttribute::MaxConstant(*Out.FrameAttribute, *A.FrameAttribute, { B, B, B });
+	}
+	return Out;
+}
+
+FAnimDatabaseFrameAttribute UAnimDatabaseFrameAttributeLibrary::FrameAttributeMinFloat(const FAnimDatabaseFrameAttribute& A, const float B)
+{
+	if (!A.IsValid())
+	{
+		UE_LOGFMT(LogAnimDatabase, Error, "FrameAttributeMinFloat: Invalid FrameAttribute.");
+		return FAnimDatabaseFrameAttribute();
+	}
+
+	if (A.Type != EAnimDatabaseAttributeType::Float &&
+		A.Type != EAnimDatabaseAttributeType::Location)
+	{
+		UE_LOGFMT(LogAnimDatabase, Error, "FrameAttributeMinFloat: FrameAttribute Type not valid. Got {Type}.", AttributeTypeNameInternal(A.Type));
+		return FAnimDatabaseFrameAttribute();
+	}
+
+	FAnimDatabaseFrameAttribute Out = MakeEmptyFrameAttribute();
+	Out.Type = A.Type;
+	if (A.Type == EAnimDatabaseAttributeType::Float)
+	{
+		UE::Learning::FrameAttribute::MinConstant(*Out.FrameAttribute, *A.FrameAttribute, { B });
+	}
+	else
+	{
+		UE::Learning::FrameAttribute::MinConstant(*Out.FrameAttribute, *A.FrameAttribute, { B, B, B });
+	}
+	return Out;
+}
+
 FAnimDatabaseFrameAttribute UAnimDatabaseFrameAttributeLibrary::FrameAttributeAddVector(const FAnimDatabaseFrameAttribute& A, const FVector B)
 {
 	if (!A.IsValid())
@@ -7290,7 +7602,7 @@ FAnimDatabaseFrameAttribute UAnimDatabaseFrameAttributeLibrary::FrameAttributeTr
 	return Out;
 }
 
-FAnimDatabaseFrameAttribute UAnimDatabaseFrameAttributeLibrary::FrameAttributeApplyTransformLocation(const FAnimDatabaseFrameAttribute& A, const FVector& B)
+FAnimDatabaseFrameAttribute UAnimDatabaseFrameAttributeLibrary::FrameAttributeApplyTransformLocation(const FAnimDatabaseFrameAttribute& A, const FVector B)
 {
 	TRACE_CPUPROFILER_EVENT_SCOPE(UAnimDatabaseFrameAttributeLibrary::FrameAttributeApplyTransformLocation);
 
@@ -7312,7 +7624,7 @@ FAnimDatabaseFrameAttribute UAnimDatabaseFrameAttributeLibrary::FrameAttributeAp
 	return Out;
 }
 
-FAnimDatabaseFrameAttribute UAnimDatabaseFrameAttributeLibrary::FrameAttributeApplyTransformDirection(const FAnimDatabaseFrameAttribute& A, const FVector& B)
+FAnimDatabaseFrameAttribute UAnimDatabaseFrameAttributeLibrary::FrameAttributeApplyTransformDirection(const FAnimDatabaseFrameAttribute& A, const FVector B)
 {
 	TRACE_CPUPROFILER_EVENT_SCOPE(UAnimDatabaseFrameAttributeLibrary::FrameAttributeApplyTransformDirection);
 
@@ -8843,6 +9155,43 @@ FAnimDatabaseFrameAttribute UAnimDatabaseFrameAttributeLibrary::FrameAttributeWe
 	return Out;
 }
 
+FAnimDatabaseFrameAttribute UAnimDatabaseFrameAttributeLibrary::FrameAttributeDynamicWeightedSum(const TArray<FAnimDatabaseFrameAttribute>& FrameAttributes, const TArray<FAnimDatabaseFrameAttribute>& Weights)
+{
+	return FrameAttributeDynamicWeightedSumFromArrayView(FrameAttributes, Weights);
+}
+
+FAnimDatabaseFrameAttribute UAnimDatabaseFrameAttributeLibrary::FrameAttributeDynamicWeightedSumFromArrayView(const TArrayView<const FAnimDatabaseFrameAttribute> FrameAttributes, const TArrayView<const FAnimDatabaseFrameAttribute> Weights)
+{
+	const int32 FrameAttributeNum = FrameAttributes.Num();
+	const int32 WeightNum = Weights.Num();
+
+	if (FrameAttributeNum == 0)
+	{
+		UE_LOGFMT(LogAnimDatabase, Warning, "FrameAttributeDynamicWeightedSumFromArrayView: No FrameAttributes Provided.");
+		return FAnimDatabaseFrameAttribute();
+	}
+
+	if (WeightNum != FrameAttributeNum)
+	{
+		UE_LOGFMT(LogAnimDatabase, Error, "FrameAttributeDynamicWeightedSumFromArrayView: Number of Frame Attributes {NumAtr} does not match number of weights {NumWeight}.",
+			FrameAttributeNum, WeightNum);
+		return FAnimDatabaseFrameAttribute();
+	}
+
+	if (FrameAttributeNum == 1)
+	{
+		return FrameAttributeMultiply(FrameAttributes[0], Weights[0]);
+	}
+
+	// TODO: Optimize this properly
+	FAnimDatabaseFrameAttribute Out = FrameAttributeMultiply(FrameAttributes[0], Weights[0]);
+	for (int32 FrameAttributeIdx = 1; FrameAttributeIdx < FrameAttributeNum; FrameAttributeIdx++)
+	{
+		Out = FrameAttributeAdd(Out, FrameAttributeMultiply(FrameAttributes[FrameAttributeIdx], Weights[FrameAttributeIdx]));
+	}
+	return Out;
+}
+
 FAnimDatabaseFrameAttribute UAnimDatabaseFrameAttributeLibrary::FrameAttributeProduct(const TArray<FAnimDatabaseFrameAttribute>& FrameAttributes)
 {
 	return FrameAttributeProductFromArrayView(FrameAttributes);
@@ -9685,10 +10034,10 @@ FAnimDatabaseFrameAttribute UAnimDatabaseFrameAttributeLibrary::FrameAttributeFi
 	const TArray<FAnimDatabaseFrameAttribute>& FitForwardDirections,
 	const TArray<FAnimDatabaseFrameAttribute>& FitRightDirections,
 	const TArray<FAnimDatabaseFrameAttribute>& FitUpDirections,
-	const TArray<float>& FitLocationWeights,
-	const TArray<float>& FitForwardDirectionWeights,
-	const TArray<float>& FitRightDirectionWeights,
-	const TArray<float>& FitUpDirectionWeights)
+	const TArray<FAnimDatabaseFrameAttribute>& FitLocationWeights,
+	const TArray<FAnimDatabaseFrameAttribute>& FitForwardDirectionWeights,
+	const TArray<FAnimDatabaseFrameAttribute>& FitRightDirectionWeights,
+	const TArray<FAnimDatabaseFrameAttribute>& FitUpDirectionWeights)
 {
 	return FrameAttributeFitTransformFromArrayViews(
 		Database,
@@ -9710,10 +10059,10 @@ FAnimDatabaseFrameAttribute UAnimDatabaseFrameAttributeLibrary::FrameAttributeFi
 	const TArrayView<const FAnimDatabaseFrameAttribute> FitForwardDirections,
 	const TArrayView<const FAnimDatabaseFrameAttribute> FitRightDirections,
 	const TArrayView<const FAnimDatabaseFrameAttribute> FitUpDirections,
-	const TArrayView<const float> FitLocationWeights,
-	const TArrayView<const float> FitForwardDirectionWeights,
-	const TArrayView<const float> FitRightDirectionWeights,
-	const TArrayView<const float> FitUpDirectionWeights)
+	const TArrayView<const FAnimDatabaseFrameAttribute> FitLocationWeights,
+	const TArrayView<const FAnimDatabaseFrameAttribute> FitForwardDirectionWeights,
+	const TArrayView<const FAnimDatabaseFrameAttribute> FitRightDirectionWeights,
+	const TArrayView<const FAnimDatabaseFrameAttribute> FitUpDirectionWeights)
 {
 	const int32 LocationNum = FitLocations.Num();
 	const int32 ForwardNum = FitForwardDirections.Num();
@@ -9730,66 +10079,69 @@ FAnimDatabaseFrameAttribute UAnimDatabaseFrameAttributeLibrary::FrameAttributeFi
 	RightAttributes.Reserve(RightNum);
 	UpAttributes.Reserve(UpNum);
 
-	TArray<float> LocationWeights; LocationWeights.SetNumZeroed(LocationNum);
-	TArray<float> ForwardWeights; ForwardWeights.SetNumZeroed(ForwardNum);
-	TArray<float> RightWeights; RightWeights.SetNumZeroed(RightNum);
-	TArray<float> UpWeights; UpWeights.SetNumZeroed(UpNum);
+	TArray<FAnimDatabaseFrameAttribute> LocationWeights; LocationWeights.SetNum(LocationNum);
+	TArray<FAnimDatabaseFrameAttribute> ForwardWeights; ForwardWeights.SetNum(ForwardNum);
+	TArray<FAnimDatabaseFrameAttribute> RightWeights; RightWeights.SetNum(RightNum);
+	TArray<FAnimDatabaseFrameAttribute> UpWeights; UpWeights.SetNum(UpNum);
+
+	const FAnimDatabaseFrameAttribute ConstantZero = UAnimDatabaseFrameAttributeLibrary::MakeFloatFrameAttributeFromConstant(FrameRanges, 0.0f);
+	const FAnimDatabaseFrameAttribute ConstantOne = UAnimDatabaseFrameAttributeLibrary::MakeFloatFrameAttributeFromConstant(FrameRanges, 1.0f);
 
-	float LocationWeightSum = 0.0f;
+	FAnimDatabaseFrameAttribute LocationWeightSum = ConstantZero;
 	for (int32 LocationIdx = 0; LocationIdx < LocationNum; LocationIdx++)
 	{
-		LocationWeights[LocationIdx] = LocationIdx < FitLocationWeights.Num() ? FitLocationWeights[LocationIdx] : 1.0f;
-		LocationWeightSum += LocationWeights[LocationIdx];
+		LocationWeights[LocationIdx] = LocationIdx < FitLocationWeights.Num() ? FitLocationWeights[LocationIdx] : ConstantOne;
+		LocationWeightSum = UAnimDatabaseFrameAttributeLibrary::FrameAttributeAdd(LocationWeightSum, LocationWeights[LocationIdx]);
 	}
 
 	for (int32 LocationIdx = 0; LocationIdx < LocationNum; LocationIdx++)
 	{
-		LocationWeights[LocationIdx] /= FMath::Max(LocationWeightSum, UE_SMALL_NUMBER);
+		LocationWeights[LocationIdx] = UAnimDatabaseFrameAttributeLibrary::FrameAttributeDivide(LocationWeights[LocationIdx], UAnimDatabaseFrameAttributeLibrary::FrameAttributeMaxFloat(LocationWeightSum, UE_SMALL_NUMBER));
 	}
 
-	float ForwardWeightSum = 0.0f;
+	FAnimDatabaseFrameAttribute ForwardWeightSum = ConstantZero;
 	for (int32 DirectionIdx = 0; DirectionIdx < ForwardNum; DirectionIdx++)
 	{
-		ForwardWeights[DirectionIdx] = DirectionIdx < FitForwardDirectionWeights.Num() ? FitForwardDirectionWeights[DirectionIdx] : 1.0f;
-		ForwardWeightSum += ForwardWeights[DirectionIdx];
+		ForwardWeights[DirectionIdx] = DirectionIdx < FitForwardDirectionWeights.Num() ? FitForwardDirectionWeights[DirectionIdx] : ConstantOne;
+		ForwardWeightSum = UAnimDatabaseFrameAttributeLibrary::FrameAttributeAdd(ForwardWeightSum, ForwardWeights[DirectionIdx]);
 		ForwardAttributes.Add(UAnimDatabaseFrameAttributeLibrary::FrameAttributeDirectionToLocation(FitForwardDirections[DirectionIdx]));
 	}
 
 	for (int32 DirectionIdx = 0; DirectionIdx < ForwardNum; DirectionIdx++)
 	{
-		ForwardWeights[DirectionIdx] /= FMath::Max(ForwardWeightSum, UE_SMALL_NUMBER);
+		ForwardWeights[DirectionIdx] = UAnimDatabaseFrameAttributeLibrary::FrameAttributeDivide(ForwardWeights[DirectionIdx], UAnimDatabaseFrameAttributeLibrary::FrameAttributeMaxFloat(ForwardWeightSum, UE_SMALL_NUMBER));
 	}
 
-	float RightWeightSum = 0.0f;
+	FAnimDatabaseFrameAttribute RightWeightSum = ConstantZero;
 	for (int32 DirectionIdx = 0; DirectionIdx < RightNum; DirectionIdx++)
 	{
-		RightWeights[DirectionIdx] = DirectionIdx < FitRightDirectionWeights.Num() ? FitRightDirectionWeights[DirectionIdx] : 1.0f;
-		RightWeightSum += RightWeights[DirectionIdx];
+		RightWeights[DirectionIdx] = DirectionIdx < FitRightDirectionWeights.Num() ? FitRightDirectionWeights[DirectionIdx] : ConstantOne;
+		RightWeightSum = UAnimDatabaseFrameAttributeLibrary::FrameAttributeAdd(RightWeightSum, RightWeights[DirectionIdx]);
 		RightAttributes.Add(UAnimDatabaseFrameAttributeLibrary::FrameAttributeDirectionToLocation(FitRightDirections[DirectionIdx]));
 	}
 
 	for (int32 DirectionIdx = 0; DirectionIdx < RightNum; DirectionIdx++)
 	{
-		RightWeights[DirectionIdx] /= FMath::Max(RightWeightSum, UE_SMALL_NUMBER);
+		RightWeights[DirectionIdx] = UAnimDatabaseFrameAttributeLibrary::FrameAttributeDivide(RightWeights[DirectionIdx], UAnimDatabaseFrameAttributeLibrary::FrameAttributeMaxFloat(RightWeightSum, UE_SMALL_NUMBER));
 	}
 
-	float UpWeightSum = 0.0f;
+	FAnimDatabaseFrameAttribute UpWeightSum = ConstantZero;
 	for (int32 DirectionIdx = 0; DirectionIdx < UpNum; DirectionIdx++)
 	{
-		UpWeights[DirectionIdx] = DirectionIdx < FitUpDirectionWeights.Num() ? FitUpDirectionWeights[DirectionIdx] : 1.0f;
-		UpWeightSum += UpWeights[DirectionIdx];
+		UpWeights[DirectionIdx] = DirectionIdx < FitUpDirectionWeights.Num() ? FitUpDirectionWeights[DirectionIdx] : ConstantOne;
+		UpWeightSum = UAnimDatabaseFrameAttributeLibrary::FrameAttributeAdd(UpWeightSum, UpWeights[DirectionIdx]);
 		UpAttributes.Add(UAnimDatabaseFrameAttributeLibrary::FrameAttributeDirectionToLocation(FitUpDirections[DirectionIdx]));
 	}
 
 	for (int32 DirectionIdx = 0; DirectionIdx < UpNum; DirectionIdx++)
 	{
-		UpWeights[DirectionIdx] /= FMath::Max(UpWeightSum, UE_SMALL_NUMBER);
+		UpWeights[DirectionIdx] = UAnimDatabaseFrameAttributeLibrary::FrameAttributeDivide(UpWeights[DirectionIdx], UAnimDatabaseFrameAttributeLibrary::FrameAttributeMaxFloat(UpWeightSum, UE_SMALL_NUMBER));
 	}
 
-	const FAnimDatabaseFrameAttribute LocationAverage = UAnimDatabaseFrameAttributeLibrary::FrameAttributeWeightedSumFromArrayView(FitLocations, LocationWeights);
-	const FAnimDatabaseFrameAttribute ForwardAverage = UAnimDatabaseFrameAttributeLibrary::FrameAttributeLocationToDirection(UAnimDatabaseFrameAttributeLibrary::FrameAttributeWeightedSumFromArrayView(ForwardAttributes, ForwardWeights));
-	const FAnimDatabaseFrameAttribute RightAverage = UAnimDatabaseFrameAttributeLibrary::FrameAttributeLocationToDirection(UAnimDatabaseFrameAttributeLibrary::FrameAttributeWeightedSumFromArrayView(RightAttributes, RightWeights));
-	const FAnimDatabaseFrameAttribute UpAverage = UAnimDatabaseFrameAttributeLibrary::FrameAttributeLocationToDirection(UAnimDatabaseFrameAttributeLibrary::FrameAttributeWeightedSumFromArrayView(UpAttributes, UpWeights));
+	const FAnimDatabaseFrameAttribute LocationAverage = UAnimDatabaseFrameAttributeLibrary::FrameAttributeDynamicWeightedSumFromArrayView(FitLocations, LocationWeights);
+	const FAnimDatabaseFrameAttribute ForwardAverage = UAnimDatabaseFrameAttributeLibrary::FrameAttributeLocationToDirection(UAnimDatabaseFrameAttributeLibrary::FrameAttributeDynamicWeightedSumFromArrayView(ForwardAttributes, ForwardWeights));
+	const FAnimDatabaseFrameAttribute RightAverage = UAnimDatabaseFrameAttributeLibrary::FrameAttributeLocationToDirection(UAnimDatabaseFrameAttributeLibrary::FrameAttributeDynamicWeightedSumFromArrayView(RightAttributes, RightWeights));
+	const FAnimDatabaseFrameAttribute UpAverage = UAnimDatabaseFrameAttributeLibrary::FrameAttributeLocationToDirection(UAnimDatabaseFrameAttributeLibrary::FrameAttributeDynamicWeightedSumFromArrayView(UpAttributes, UpWeights));
 
 	const FAnimDatabaseFrameAttribute RotationAverage = UAnimDatabaseFrameAttributeLibrary::MakeRotationFrameAttributeFromBasisDirections(ForwardAverage, RightAverage, UpAverage);
 
@@ -9832,14 +10184,15 @@ FAnimDatabaseFrameAttribute UAnimDatabaseFrameAttributeFunction_ContactCurve::Ma
 
 FAnimDatabaseFrameAttribute UAnimDatabaseFrameAttributeFunction_FitTransform::MakeFrameAttribute_Implementation(const UAnimDatabase* Database, const FAnimDatabaseFrameRanges& FrameRanges) const
 {
-	TArray<FAnimDatabaseFrameAttribute> LocationAttributes;
-	TArray<FAnimDatabaseFrameAttribute> ForwardAttributes;
-	TArray<FAnimDatabaseFrameAttribute> RightAttributes;
-	TArray<FAnimDatabaseFrameAttribute> UpAttributes;
+	TArray<FAnimDatabaseFrameAttribute> LocationAttributes, ForwardAttributes, RightAttributes, UpAttributes, LocationWeightAttributes, ForwardWeightAttributes, RightWeightAttributes, UpWeightAttributes;
 	UAnimDatabaseFrameAttributeLibrary::MakeFrameAttributesFromFunctions(LocationAttributes, Database, FrameRanges, FitLocations);
 	UAnimDatabaseFrameAttributeLibrary::MakeFrameAttributesFromFunctions(ForwardAttributes, Database, FrameRanges, FitForwardDirections);
 	UAnimDatabaseFrameAttributeLibrary::MakeFrameAttributesFromFunctions(RightAttributes, Database, FrameRanges, FitRightDirections);
 	UAnimDatabaseFrameAttributeLibrary::MakeFrameAttributesFromFunctions(UpAttributes, Database, FrameRanges, FitUpDirections);
+	UAnimDatabaseFrameAttributeLibrary::MakeFrameAttributesFromFunctions(LocationWeightAttributes, Database, FrameRanges, FitLocationWeights);
+	UAnimDatabaseFrameAttributeLibrary::MakeFrameAttributesFromFunctions(ForwardWeightAttributes, Database, FrameRanges, FitForwardDirectionWeights);
+	UAnimDatabaseFrameAttributeLibrary::MakeFrameAttributesFromFunctions(RightWeightAttributes, Database, FrameRanges, FitRightDirectionWeights);
+	UAnimDatabaseFrameAttributeLibrary::MakeFrameAttributesFromFunctions(UpWeightAttributes, Database, FrameRanges, FitUpDirectionWeights);
 
 	return UAnimDatabaseFrameAttributeLibrary::FrameAttributeFitTransformFromArrayViews(
 		Database,
@@ -9848,10 +10201,10 @@ FAnimDatabaseFrameAttribute UAnimDatabaseFrameAttributeFunction_FitTransform::Ma
 		ForwardAttributes,
 		RightAttributes,
 		UpAttributes,
-		FitLocationWeights,
-		FitForwardDirectionWeights,
-		FitRightDirectionWeights,
-		FitUpDirectionWeights);
+		LocationWeightAttributes,
+		ForwardWeightAttributes,
+		RightWeightAttributes,
+		UpWeightAttributes);
 }
 
 FAnimDatabaseFrameAttribute UAnimDatabaseFrameAttributeFunction_BoneGlobalTransform::MakeFrameAttribute_Implementation(const UAnimDatabase* Database, const FAnimDatabaseFrameRanges& FrameRanges) const
diff --git a/Engine/Plugins/Experimental/Animation/AnimDatabase/Source/AnimDatabase/Private/AnimDatabaseFrameRanges.cpp b/Engine/Plugins/Experimental/Animation/AnimDatabase/Source/AnimDatabase/Private/AnimDatabaseFrameRanges.cpp
--- a/Engine/Plugins/Experimental/Animation/AnimDatabase/Source/AnimDatabase/Private/AnimDatabaseFrameRanges.cpp
+++ b/Engine/Plugins/Experimental/Animation/AnimDatabase/Source/AnimDatabase/Private/AnimDatabaseFrameRanges.cpp
@@ -903,6 +903,158 @@ FAnimDatabaseFrameRanges UAnimDatabaseFrameRangesLibrary::MakeFrameRangesFromMir
 	return Out;
 }
 
+FAnimDatabaseFrameRanges UAnimDatabaseFrameRangesLibrary::MakeFrameRangesFromRootMotionEnabled(const UAnimDatabase* Database, const FAnimDatabaseFrameRanges& FrameRanges)
+{
+	if (!Database)
+	{
+		UE_LOGFMT(LogAnimDatabase, Error, "MakeFrameRangesFromRootMotionEnabled: Database is nullptr.");
+		return FAnimDatabaseFrameRanges();
+	}
+
+	if (!FrameRanges.IsValid())
+	{
+		UE_LOGFMT(LogAnimDatabase, Error, "MakeFrameRangesFromRootMotionEnabled: Invalid FrameRanges Object.");
+		return FAnimDatabaseFrameRanges();
+	}
+
+	FAnimDatabaseFrameRanges Out = MakeEmptyFrameRanges();
+
+	const int32 EntryNum = FrameRanges.FrameRangeSet->GetEntryNum();
+
+	for (int32 EntryIdx = 0; EntryIdx < EntryNum; EntryIdx++)
+	{
+		const int32 SequenceIdx = FrameRanges.FrameRangeSet->GetEntrySequence(EntryIdx);
+
+		if (const UAnimSequence* AnimSequence = Database->GetAnimSequence(SequenceIdx))
+		{
+			if (AnimSequence->bEnableRootMotion)
+			{
+				Out.FrameRangeSet->AddEntry(
+					SequenceIdx,
+					FrameRanges.FrameRangeSet->GetEntryRangeStarts(EntryIdx),
+					FrameRanges.FrameRangeSet->GetEntryRangeLengths(EntryIdx));
+			}
+		}
+	}
+
+	Out.FrameRangeSet->Check();
+	return Out;
+}
+
+FAnimDatabaseFrameRanges UAnimDatabaseFrameRangesLibrary::MakeFrameRangesFromRootMotionDisabled(const UAnimDatabase* Database, const FAnimDatabaseFrameRanges& FrameRanges)
+{
+	if (!Database)
+	{
+		UE_LOGFMT(LogAnimDatabase, Error, "MakeFrameRangesFromRootMotionDisabled: Database is nullptr.");
+		return FAnimDatabaseFrameRanges();
+	}
+
+	if (!FrameRanges.IsValid())
+	{
+		UE_LOGFMT(LogAnimDatabase, Error, "MakeFrameRangesFromRootMotionDisabled: Invalid FrameRanges Object.");
+		return FAnimDatabaseFrameRanges();
+	}
+
+	FAnimDatabaseFrameRanges Out = MakeEmptyFrameRanges();
+
+	const int32 EntryNum = FrameRanges.FrameRangeSet->GetEntryNum();
+
+	for (int32 EntryIdx = 0; EntryIdx < EntryNum; EntryIdx++)
+	{
+		const int32 SequenceIdx = FrameRanges.FrameRangeSet->GetEntrySequence(EntryIdx);
+
+		if (const UAnimSequence* AnimSequence = Database->GetAnimSequence(SequenceIdx))
+		{
+			if (!AnimSequence->bEnableRootMotion)
+			{
+				Out.FrameRangeSet->AddEntry(
+					SequenceIdx,
+					FrameRanges.FrameRangeSet->GetEntryRangeStarts(EntryIdx),
+					FrameRanges.FrameRangeSet->GetEntryRangeLengths(EntryIdx));
+			}
+		}
+	}
+
+	Out.FrameRangeSet->Check();
+	return Out;
+}
+
+FAnimDatabaseFrameRanges UAnimDatabaseFrameRangesLibrary::MakeFrameRangesFromForceRootLockEnabled(const UAnimDatabase* Database, const FAnimDatabaseFrameRanges& FrameRanges)
+{
+	if (!Database)
+	{
+		UE_LOGFMT(LogAnimDatabase, Error, "MakeFrameRangesFromRootMotionDisabled: Database is nullptr.");
+		return FAnimDatabaseFrameRanges();
+	}
+
+	if (!FrameRanges.IsValid())
+	{
+		UE_LOGFMT(LogAnimDatabase, Error, "MakeFrameRangesFromRootMotionDisabled: Invalid FrameRanges Object.");
+		return FAnimDatabaseFrameRanges();
+	}
+
+	FAnimDatabaseFrameRanges Out = MakeEmptyFrameRanges();
+
+	const int32 EntryNum = FrameRanges.FrameRangeSet->GetEntryNum();
+
+	for (int32 EntryIdx = 0; EntryIdx < EntryNum; EntryIdx++)
+	{
+		const int32 SequenceIdx = FrameRanges.FrameRangeSet->GetEntrySequence(EntryIdx);
+
+		if (const UAnimSequence* AnimSequence = Database->GetAnimSequence(SequenceIdx))
+		{
+			if (AnimSequence->bForceRootLock)
+			{
+				Out.FrameRangeSet->AddEntry(
+					SequenceIdx,
+					FrameRanges.FrameRangeSet->GetEntryRangeStarts(EntryIdx),
+					FrameRanges.FrameRangeSet->GetEntryRangeLengths(EntryIdx));
+			}
+		}
+	}
+
+	Out.FrameRangeSet->Check();
+	return Out;
+}
+
+FAnimDatabaseFrameRanges UAnimDatabaseFrameRangesLibrary::MakeFrameRangesFromForceRootLockDisabled(const UAnimDatabase* Database, const FAnimDatabaseFrameRanges& FrameRanges)
+{
+	if (!Database)
+	{
+		UE_LOGFMT(LogAnimDatabase, Error, "MakeFrameRangesFromRootMotionDisabled: Database is nullptr.");
+		return FAnimDatabaseFrameRanges();
+	}
+
+	if (!FrameRanges.IsValid())
+	{
+		UE_LOGFMT(LogAnimDatabase, Error, "MakeFrameRangesFromRootMotionDisabled: Invalid FrameRanges Object.");
+		return FAnimDatabaseFrameRanges();
+	}
+
+	FAnimDatabaseFrameRanges Out = MakeEmptyFrameRanges();
+
+	const int32 EntryNum = FrameRanges.FrameRangeSet->GetEntryNum();
+
+	for (int32 EntryIdx = 0; EntryIdx < EntryNum; EntryIdx++)
+	{
+		const int32 SequenceIdx = FrameRanges.FrameRangeSet->GetEntrySequence(EntryIdx);
+
+		if (const UAnimSequence* AnimSequence = Database->GetAnimSequence(SequenceIdx))
+		{
+			if (!AnimSequence->bForceRootLock)
+			{
+				Out.FrameRangeSet->AddEntry(
+					SequenceIdx,
+					FrameRanges.FrameRangeSet->GetEntryRangeStarts(EntryIdx),
+					FrameRanges.FrameRangeSet->GetEntryRangeLengths(EntryIdx));
+			}
+		}
+	}
+
+	Out.FrameRangeSet->Check();
+	return Out;
+}
+
 FAnimDatabaseFrameRanges UAnimDatabaseFrameRangesLibrary::MakeFrameRangesFromSequenceIndex(const UAnimDatabase* Database, const int32 SequenceIdx)
 {
 	if (!Database)
@@ -2564,6 +2716,26 @@ FAnimDatabaseFrameRanges UAnimDatabaseFrameRangesFunction_NotMirrored::MakeFrame
 	return UAnimDatabaseFrameRangesLibrary::MakeFrameRangesFromNotMirrored(Database, FrameRanges);
 }
 
+FAnimDatabaseFrameRanges UAnimDatabaseFrameRangesFunction_RootMotionEnabled::MakeFrameRanges_Implementation(const UAnimDatabase* Database, const FAnimDatabaseFrameRanges& FrameRanges) const
+{
+	return UAnimDatabaseFrameRangesLibrary::MakeFrameRangesFromRootMotionEnabled(Database, FrameRanges);
+}
+
+FAnimDatabaseFrameRanges UAnimDatabaseFrameRangesFunction_RootMotionDisabled::MakeFrameRanges_Implementation(const UAnimDatabase* Database, const FAnimDatabaseFrameRanges& FrameRanges) const
+{
+	return UAnimDatabaseFrameRangesLibrary::MakeFrameRangesFromRootMotionDisabled(Database, FrameRanges);
+}
+
+FAnimDatabaseFrameRanges UAnimDatabaseFrameRangesFunction_ForceRootLockEnabled::MakeFrameRanges_Implementation(const UAnimDatabase* Database, const FAnimDatabaseFrameRanges& FrameRanges) const
+{
+	return UAnimDatabaseFrameRangesLibrary::MakeFrameRangesFromForceRootLockEnabled(Database, FrameRanges);
+}
+
+FAnimDatabaseFrameRanges UAnimDatabaseFrameRangesFunction_ForceRootLockDisabled::MakeFrameRanges_Implementation(const UAnimDatabase* Database, const FAnimDatabaseFrameRanges& FrameRanges) const
+{
+	return UAnimDatabaseFrameRangesLibrary::MakeFrameRangesFromForceRootLockDisabled(Database, FrameRanges);
+}
+
 FAnimDatabaseFrameRanges UAnimDatabaseFrameRangesFunction_AnimNotify::MakeFrameRanges_Implementation(const UAnimDatabase* Database, const FAnimDatabaseFrameRanges& FrameRanges) const
 {
 	return UAnimDatabaseFrameRangesLibrary::MakeFrameRangesFromAnimNotify(Database, FrameRanges, AnimNotify);
diff --git a/Engine/Plugins/Experimental/Animation/AnimDatabase/Source/AnimDatabase/Public/AnimDatabase.h b/Engine/Plugins/Experimental/Animation/AnimDatabase/Source/AnimDatabase/Public/AnimDatabase.h
--- a/Engine/Plugins/Experimental/Animation/AnimDatabase/Source/AnimDatabase/Public/AnimDatabase.h
+++ b/Engine/Plugins/Experimental/Animation/AnimDatabase/Source/AnimDatabase/Public/AnimDatabase.h
@@ -46,6 +46,7 @@ class USkeletalMesh;
 class UAnimSequence;
 class UMirrorDataTable;
 class IAnimationDataController;
+class UStaticMesh;
 
 //--------------------------------------------------
 
@@ -840,11 +841,11 @@ public:
 public:
 
 	/** Editor-only object containing the viewport settings for the editor window */
-	UPROPERTY()
+	UPROPERTY(Instanced)
 	TObjectPtr<UAnimDatabaseViewportSettings> ViewportSettings;
 
 	/** Editor-only object containing the query state for the editor window */
-	UPROPERTY()
+	UPROPERTY(Instanced)
 	TObjectPtr<UAnimDatabaseQuery> Query;
 };
 
@@ -997,6 +998,15 @@ public:
 	static UE_API void DatabaseSetLocalBoneTransforms(UAnimDatabase* Database, const TArray<FName>& BoneNames, const TArray<FAnimDatabaseFrameAttribute>& TransformFrameAttributes, const bool bShouldTransact = false);
 	static UE_API void DatabaseSetLocalBoneTransformsFromArrayViews(UAnimDatabase* Database, const TArrayView<const FName> BoneNames, const TArrayView<const FAnimDatabaseFrameAttribute> TransformFrameAttributes, const bool bShouldTransact = false);
 
+	/** Set the given bone's global transform in the database at the given transform frame attribute */
+	UFUNCTION(BlueprintCallable, Category = "AnimDatabase", meta = (DevelopmentOnly))
+	static UE_API void DatabaseSetGlobalBoneTransform(UAnimDatabase* Database, const FName BoneName, const FAnimDatabaseFrameAttribute& TransformFrameAttribute, const bool bShouldTransact = false);
+
+	/** Set the given bones' global transforms in the database at the given transform frame attributes */
+	UFUNCTION(BlueprintCallable, Category = "AnimDatabase", meta = (DevelopmentOnly))
+	static UE_API void DatabaseSetGlobalBoneTransforms(UAnimDatabase* Database, const TArray<FName>& BoneNames, const TArray<FAnimDatabaseFrameAttribute>& TransformFrameAttributes, const bool bShouldTransact = false);
+	static UE_API void DatabaseSetGlobalBoneTransformsFromArrayViews(UAnimDatabase* Database, const TArrayView<const FName> BoneNames, const TArrayView<const FAnimDatabaseFrameAttribute> TransformFrameAttributes, const bool bShouldTransact = false);
+
 	/** Strips unused bone tracks from the data to reduce the file size on disk */
 	UFUNCTION(BlueprintCallable, Category = "AnimDatabase", meta = (DevelopmentOnly))
 	static UE_API void DatabaseStripUnusedBoneTracks(UAnimDatabase* Database, const FAnimDatabaseFrameRanges& FrameRanges, const bool bShouldTransact = false);
@@ -1899,7 +1909,7 @@ public:
 
 	/** Adjustment to apply to the pelvis height */
 	UPROPERTY(EditAnywhere, Category = "Settings", meta = (ForceUnits = "cm"))
-	float PelvisHeightAdjustment = -1.0f;
+	float PelvisHeightAdjustment = 0.0f;
 
 #endif
 
@@ -2277,6 +2287,70 @@ public:
 #endif
 };
 
+
+/** Debug Draw a bone attachment such as a prop */
+UCLASS(BlueprintType, Blueprintable, meta = (DisplayName = "Bone Attachment"))
+class UAnimDatabaseDebugDraw_BoneAttachment : public UAnimDatabaseDebugDraw
+{
+	GENERATED_BODY()
+
+#if WITH_EDITORONLY_DATA
+
+public:
+
+	/** Bone attached to */
+	UPROPERTY(EditAnywhere, Category = "Settings")
+	FName BoneName = NAME_None;
+
+	/** Relative transform of the attachment to the bone */
+	UPROPERTY(EditAnywhere, Category = "Settings")
+	FTransform RelativeTransform = FTransform::Identity;
+
+	/** If to draw the attachment transform */
+	UPROPERTY(EditAnywhere, Category = "Settings")
+	bool bDrawTransform = true;
+
+	/** Transform Radius */
+	UPROPERTY(EditAnywhere, Category = "Settings", meta = (ForceForceUnits = "cm"))
+	float TransformRadius = 25.0f;
+
+	/** Transform Thickness */
+	UPROPERTY(EditAnywhere, Category = "Settings", meta = (ForceForceUnits = "cm"))
+	float TransformThickness = 1.0f;
+
+	/** Static mesh to use for attachment (will draw bounding box only) */
+	UPROPERTY(EditAnywhere, Category = "Settings")
+	TObjectPtr<UStaticMesh> StaticMesh;
+
+	/** Bounding box Thickness */
+	UPROPERTY(EditAnywhere, Category = "Settings", meta = (ForceForceUnits = "cm"))
+	float BoundingBoxThickness = 1.0f;
+
+	/** If to depth test while drawing */
+	UPROPERTY(EditAnywhere, Category = "Settings")
+	bool bDepthTest = true;
+
+#endif
+
+#if WITH_EDITOR
+
+public:
+
+	virtual void DrawDebug_Implementation(
+		const FDebugDrawer& Drawer,
+		const FAnimDatabasePoseState& InPoseState,
+		const UAnimDatabase* InDatabase,
+		const FAnimDatabaseFrameRanges& InFrameRanges,
+		const FTransform& RangeViewportTransform,
+		const int32 SequenceIdx,
+		const float SequenceTime,
+		const int32 RangeStart,
+		const int32 RangeLength,
+		const FLinearColor& IdentifierColor) override;
+
+#endif
+};
+
 /** Debug Draw class for drawing root velocities */
 UCLASS(BlueprintType, Blueprintable, meta = (DisplayName = "Root Velocity"))
 class UAnimDatabaseDebugDraw_RootVelocity : public UAnimDatabaseDebugDraw
diff --git a/Engine/Plugins/Experimental/Animation/AnimDatabase/Source/AnimDatabase/Public/AnimDatabaseFrameAttribute.h b/Engine/Plugins/Experimental/Animation/AnimDatabase/Source/AnimDatabase/Public/AnimDatabaseFrameAttribute.h
--- a/Engine/Plugins/Experimental/Animation/AnimDatabase/Source/AnimDatabase/Public/AnimDatabaseFrameAttribute.h
+++ b/Engine/Plugins/Experimental/Animation/AnimDatabase/Source/AnimDatabase/Public/AnimDatabaseFrameAttribute.h
@@ -722,7 +722,7 @@ public:
 		const float RelativeTime = 0.0f);
 
 	static UE_API void MakeBoneGlobalTransformFrameAttributesFromNamesArrayView(
-		TArray<FAnimDatabaseFrameAttribute>& OutTransformFrameAttributes,
+		TArrayView<FAnimDatabaseFrameAttribute> OutTransformFrameAttributes,
 		const UAnimDatabase* Database,
 		const FAnimDatabaseFrameRanges& FrameRanges,
 		const TArrayView<const FName> BoneNames,
@@ -736,6 +736,16 @@ public:
 	UFUNCTION(BlueprintPure = true, Category = "AnimDatabase", meta = (NotBlueprintThreadSafe, AdvancedDisplay = 3))
 	static UE_API FAnimDatabaseFrameAttribute MakeBoneGlobalLocationFrameAttributeFromName(const UAnimDatabase* Database, const FAnimDatabaseFrameRanges& FrameRanges, const FName BoneName, const float RelativeTime = 0.0f);
 
+	/** Makes location frame attributes from the global location of bones */
+	UFUNCTION(BlueprintPure = true, Category = "AnimDatabase", meta = (NotBlueprintThreadSafe, AdvancedDisplay = 4))
+	static UE_API void MakeBoneGlobalLocationFrameAttributes(TArray<FAnimDatabaseFrameAttribute>& OutLocationFrameAttributes, const UAnimDatabase* Database, const FAnimDatabaseFrameRanges& FrameRanges, const TArray<int32>& BoneIndices, const float RelativeTime = 0.0f);
+	static UE_API void MakeBoneGlobalLocationFrameAttributesFromArrayViews(TArrayView<FAnimDatabaseFrameAttribute> OutLocationFrameAttributes, const UAnimDatabase* Database, const FAnimDatabaseFrameRanges& FrameRanges, const TArrayView<const int32> BoneIndices, const float RelativeTime = 0.0f);
+
+	/** Makes location frame attributes from the global location of bones, provided by name */
+	UFUNCTION(BlueprintPure = true, Category = "AnimDatabase", meta = (NotBlueprintThreadSafe, AdvancedDisplay = 4))
+	static UE_API void MakeBoneGlobalLocationFrameAttributesFromNames(TArray<FAnimDatabaseFrameAttribute>& OutLocationFrameAttributes, const UAnimDatabase* Database, const FAnimDatabaseFrameRanges& FrameRanges, const TArray<FName>& BoneNames, const float RelativeTime = 0.0f);
+	static UE_API void MakeBoneGlobalLocationFrameAttributesFromNamesArrayView(TArrayView<FAnimDatabaseFrameAttribute> OutLocationFrameAttributes, const UAnimDatabase* Database, const FAnimDatabaseFrameRanges& FrameRanges, const TArrayView<const FName> BoneNames, const float RelativeTime = 0.0f);
+
 	/** Makes a direction frame attribute from the global direction of a bone */
 	UFUNCTION(BlueprintPure = true, Category = "AnimDatabase", meta = (NotBlueprintThreadSafe, AdvancedDisplay = 4))
 	static UE_API FAnimDatabaseFrameAttribute MakeBoneGlobalDirectionFrameAttribute(const UAnimDatabase* Database, const FAnimDatabaseFrameRanges& FrameRanges, const int32 BoneIndex, const FVector LocalDirection = FVector::ForwardVector, const float RelativeTime = 0.0f);
@@ -744,6 +754,16 @@ public:
 	UFUNCTION(BlueprintPure = true, Category = "AnimDatabase", meta = (NotBlueprintThreadSafe, AdvancedDisplay = 4))
 	static UE_API FAnimDatabaseFrameAttribute MakeBoneGlobalDirectionFrameAttributeFromName(const UAnimDatabase* Database, const FAnimDatabaseFrameRanges& FrameRanges, const FName BoneName, const FVector LocalDirection = FVector::ForwardVector, const float RelativeTime = 0.0f);
 
+	/** Makes direction frame attributes from the global direction of bones */
+	UFUNCTION(BlueprintPure = true, Category = "AnimDatabase", meta = (NotBlueprintThreadSafe, AdvancedDisplay = 5))
+	static UE_API void MakeBoneGlobalDirectionFrameAttributes(TArray<FAnimDatabaseFrameAttribute>& OutDirectionFrameAttributes, const UAnimDatabase* Database, const FAnimDatabaseFrameRanges& FrameRanges, const TArray<int32>& BoneIndices, const FVector LocalDirection = FVector::ForwardVector, const float RelativeTime = 0.0f);
+	static UE_API void MakeBoneGlobalDirectionFrameAttributesFromArrayViews(TArrayView<FAnimDatabaseFrameAttribute> OutDirectionFrameAttributes, const UAnimDatabase* Database, const FAnimDatabaseFrameRanges& FrameRanges, const TArrayView<const int32> BoneIndices, const FVector LocalDirection = FVector::ForwardVector, const float RelativeTime = 0.0f);
+
+	/** Makes direction frame attributes from the global direction of bones, provided by name */
+	UFUNCTION(BlueprintPure = true, Category = "AnimDatabase", meta = (NotBlueprintThreadSafe, AdvancedDisplay = 5))
+	static UE_API void MakeBoneGlobalDirectionFrameAttributesFromNames(TArray<FAnimDatabaseFrameAttribute>& OutDirectionFrameAttributes, const UAnimDatabase* Database, const FAnimDatabaseFrameRanges& FrameRanges, const TArray<FName>& BoneNames, const FVector LocalDirection = FVector::ForwardVector, const float RelativeTime = 0.0f);
+	static UE_API void MakeBoneGlobalDirectionFrameAttributesFromNamesArrayView(TArrayView<FAnimDatabaseFrameAttribute> OutDirectionFrameAttributes, const UAnimDatabase* Database, const FAnimDatabaseFrameRanges& FrameRanges, const TArrayView<const FName> BoneNames, const FVector LocalDirection = FVector::ForwardVector, const float RelativeTime = 0.0f);
+
 	/** Makes a linear velocity frame attribute from the global linear velocity of a bone */
 	UFUNCTION(BlueprintPure = true, Category = "AnimDatabase", meta = (NotBlueprintThreadSafe, AdvancedDisplay = 3))
 	static UE_API FAnimDatabaseFrameAttribute MakeBoneGlobalLinearVelocityFrameAttribute(const UAnimDatabase* Database, const FAnimDatabaseFrameRanges& FrameRanges, const int32 BoneIndex, const float RelativeTime = 0.0f);
@@ -768,6 +788,14 @@ public:
 	UFUNCTION(BlueprintPure = true, Category = "AnimDatabase", meta = (AdvancedDisplay = 5))
 	static UE_API void MakeBoneGlobalLocationAndLinearVelocityFrameAttributeFromName(FAnimDatabaseFrameAttribute& OutLocationFrameAttribute, FAnimDatabaseFrameAttribute& OutLinearVelocityFrameAttribute, const UAnimDatabase* Database, const FAnimDatabaseFrameRanges& FrameRanges, const FName BoneName, const float RelativeTime = 0.0f);
 
+	/** Makes a direction frame attribute from the global vector between bones */
+	UFUNCTION(BlueprintPure = true, Category = "AnimDatabase", meta = (NotBlueprintThreadSafe, AdvancedDisplay = 4))
+	static UE_API FAnimDatabaseFrameAttribute MakeGlobalDirectionBetweenBonesFrameAttribute(const UAnimDatabase* Database, const FAnimDatabaseFrameRanges& FrameRanges, const int32 BoneIndexFrom, const int32 BoneIndexTo);
+
+	/** Makes a direction frame attribute from the global vector between bones, provided by name */
+	UFUNCTION(BlueprintPure = true, Category = "AnimDatabase", meta = (NotBlueprintThreadSafe, AdvancedDisplay = 4))
+	static UE_API FAnimDatabaseFrameAttribute MakeGlobalDirectionBetweenBonesFrameAttributeFromNames(const UAnimDatabase* Database, const FAnimDatabaseFrameRanges& FrameRanges, const FName BoneNameFrom, const FName BoneNameTo);
+
 	/** Makes arrays of frame attributes for a pose from the provided bone indices */
 	UFUNCTION(BlueprintPure = true, Category = "AnimDatabase", meta = (NotBlueprintThreadSafe, AdvancedDisplay = 5))
 	static UE_API void MakePoseFrameAttribute(
@@ -1232,6 +1260,23 @@ public:
 	UFUNCTION(BlueprintPure, Category = "AnimDatabase", meta = (CommutativeAssociativeBinaryOperator, CompactNodeTitle = "/"))
 	static UE_API FAnimDatabaseFrameAttribute FrameAttributeDivide(FAnimDatabaseFrameAttribute A, FAnimDatabaseFrameAttribute B);
 
+	/** Computes the component-wise maximum of two frame attributes. Assumes types are compatible */
+	UFUNCTION(BlueprintPure, Category = "AnimDatabase", meta = (CommutativeAssociativeBinaryOperator, CompactNodeTitle = "MAX"))
+	static UE_API FAnimDatabaseFrameAttribute FrameAttributeMax(FAnimDatabaseFrameAttribute A, FAnimDatabaseFrameAttribute B);
+
+	/** Computes the component-wise minimum of two frame attributes. Assumes types are compatible */
+	UFUNCTION(BlueprintPure, Category = "AnimDatabase", meta = (CommutativeAssociativeBinaryOperator, CompactNodeTitle = "MIN"))
+	static UE_API FAnimDatabaseFrameAttribute FrameAttributeMin(FAnimDatabaseFrameAttribute A, FAnimDatabaseFrameAttribute B);
+
+	/** Computes the dot product of two frame attributes. Assumes types are compatible */
+	UFUNCTION(BlueprintPure, Category = "AnimDatabase", meta = (CommutativeAssociativeBinaryOperator, CompactNodeTitle = "DOT"))
+	static UE_API FAnimDatabaseFrameAttribute FrameAttributeDot(FAnimDatabaseFrameAttribute A, FAnimDatabaseFrameAttribute B);
+
+	/** Computes the cross product of two frame attributes. Assumes types are compatible */
+	UFUNCTION(BlueprintPure, Category = "AnimDatabase", meta = (CommutativeAssociativeBinaryOperator, CompactNodeTitle = "CROSS"))
+	static UE_API FAnimDatabaseFrameAttribute FrameAttributeCross(FAnimDatabaseFrameAttribute A, FAnimDatabaseFrameAttribute B);
+
+
 
 	/** Adds a constant float value to a frame attribute. Assumes types are compatible */
 	UFUNCTION(BlueprintPure, Category = "AnimDatabase", meta = (CompactNodeTitle = "+"))
@@ -1249,6 +1294,14 @@ public:
 	UFUNCTION(BlueprintPure, Category = "AnimDatabase", meta = (CompactNodeTitle = "/"))
 	static UE_API FAnimDatabaseFrameAttribute FrameAttributeDivideFloat(const FAnimDatabaseFrameAttribute& A, const float B);
 
+	/** Computes the component-wise maximum of a frame attribute and a float value. Assumes types are compatible */
+	UFUNCTION(BlueprintPure, Category = "AnimDatabase", meta = (CompactNodeTitle = "MAX"))
+	static UE_API FAnimDatabaseFrameAttribute FrameAttributeMaxFloat(const FAnimDatabaseFrameAttribute& A, const float B);
+
+	/** Computes the component-wise minimum of a frame attribute and a float value. Assumes types are compatible */
+	UFUNCTION(BlueprintPure, Category = "AnimDatabase", meta = (CompactNodeTitle = "MIN"))
+	static UE_API FAnimDatabaseFrameAttribute FrameAttributeMinFloat(const FAnimDatabaseFrameAttribute& A, const float B);
+
 
 	/** Adds a constant vector value to a frame attribute. Assumes types are compatible */
 	UFUNCTION(BlueprintPure, Category = "AnimDatabase", meta = (CompactNodeTitle = "+"))
@@ -1306,11 +1359,11 @@ public:
 
 	/** Applies a transformation to a constant location vector. Assumes types are compatible */
 	UFUNCTION(BlueprintPure, Category = "AnimDatabase")
-	static UE_API FAnimDatabaseFrameAttribute FrameAttributeApplyTransformLocation(const FAnimDatabaseFrameAttribute& A, const FVector& B);
+	static UE_API FAnimDatabaseFrameAttribute FrameAttributeApplyTransformLocation(const FAnimDatabaseFrameAttribute& A, const FVector B);
 
 	/** Applies a transformation to a constant direction vector. Assumes types are compatible */
 	UFUNCTION(BlueprintPure, Category = "AnimDatabase")
-	static UE_API FAnimDatabaseFrameAttribute FrameAttributeApplyTransformDirection(const FAnimDatabaseFrameAttribute& A, const FVector& B);
+	static UE_API FAnimDatabaseFrameAttribute FrameAttributeApplyTransformDirection(const FAnimDatabaseFrameAttribute& A, const FVector B);
 
 	/** Adds a constant angle value in degrees to an angle frame attribute. */
 	UFUNCTION(BlueprintPure, Category = "AnimDatabase", meta = (CompactNodeTitle = "+"))
@@ -1362,6 +1415,11 @@ public:
 	static UE_API FAnimDatabaseFrameAttribute FrameAttributeWeightedSum(const TArray<FAnimDatabaseFrameAttribute>& FrameAttributes, const TArray<float>& Weights);
 	static UE_API FAnimDatabaseFrameAttribute FrameAttributeWeightedSumFromArrayView(const TArrayView<const FAnimDatabaseFrameAttribute> FrameAttributes, const TArrayView<const float> Weights);
 
+	/** Computes the weighted sum of multiple frame attributes where the weights are dynamic. Assumes all are the same type */
+	UFUNCTION(BlueprintPure, Category = "AnimDatabase")
+	static UE_API FAnimDatabaseFrameAttribute FrameAttributeDynamicWeightedSum(const TArray<FAnimDatabaseFrameAttribute>& FrameAttributes, const TArray<FAnimDatabaseFrameAttribute>& Weights);
+	static UE_API FAnimDatabaseFrameAttribute FrameAttributeDynamicWeightedSumFromArrayView(const TArrayView<const FAnimDatabaseFrameAttribute> FrameAttributes, const TArrayView<const FAnimDatabaseFrameAttribute> Weights);
+
 	/** Computes the product of multiple frame attributes. Assumes all are the same type */
 	UFUNCTION(BlueprintPure, Category = "AnimDatabase")
 	static UE_API FAnimDatabaseFrameAttribute FrameAttributeProduct(const TArray<FAnimDatabaseFrameAttribute>& FrameAttributes);
@@ -1862,7 +1920,7 @@ public:
 	 * @param FitRightDirectionWeights			Optional array of weights to use for the fit right directions
 	 * @param FitUpDirectionWeights				Optional array of weights to use for the fit up directions
 	 */
-	UFUNCTION(BlueprintPure, Category = "AnimDatabase", meta = (NotBlueprintThreadSafe))
+	UFUNCTION(BlueprintPure, Category = "AnimDatabase", meta = (NotBlueprintThreadSafe, AutoCreateRefTerm="FitLocationWeights, FitForwardDirectionWeights, FitRightDirectionWeights, FitUpDirectionWeights"))
 	static UE_API FAnimDatabaseFrameAttribute FrameAttributeFitTransform(
 		const UAnimDatabase* Database,
 		const FAnimDatabaseFrameRanges& FrameRanges,
@@ -1870,10 +1928,10 @@ public:
 		const TArray<FAnimDatabaseFrameAttribute>& FitForwardDirections,
 		const TArray<FAnimDatabaseFrameAttribute>& FitRightDirections,
 		const TArray<FAnimDatabaseFrameAttribute>& FitUpDirections,
-		const TArray<float>& FitLocationWeights,
-		const TArray<float>& FitForwardDirectionWeights,
-		const TArray<float>& FitRightDirectionWeights,
-		const TArray<float>& FitUpDirectionWeights);
+		const TArray<FAnimDatabaseFrameAttribute>& FitLocationWeights,
+		const TArray<FAnimDatabaseFrameAttribute>& FitForwardDirectionWeights,
+		const TArray<FAnimDatabaseFrameAttribute>& FitRightDirectionWeights,
+		const TArray<FAnimDatabaseFrameAttribute>& FitUpDirectionWeights);
 
 	static UE_API FAnimDatabaseFrameAttribute FrameAttributeFitTransformFromArrayViews(
 		const UAnimDatabase* Database,
@@ -1882,10 +1940,10 @@ public:
 		const TArrayView<const FAnimDatabaseFrameAttribute> FitForwardDirections,
 		const TArrayView<const FAnimDatabaseFrameAttribute> FitRightDirections,
 		const TArrayView<const FAnimDatabaseFrameAttribute> FitUpDirections,
-		const TArrayView<const float> FitLocationWeights,
-		const TArrayView<const float> FitForwardDirectionWeights,
-		const TArrayView<const float> FitRightDirectionWeights,
-		const TArrayView<const float> FitUpDirectionWeights);
+		const TArrayView<const FAnimDatabaseFrameAttribute> FitLocationWeights,
+		const TArrayView<const FAnimDatabaseFrameAttribute> FitForwardDirectionWeights,
+		const TArrayView<const FAnimDatabaseFrameAttribute> FitRightDirectionWeights,
+		const TArrayView<const FAnimDatabaseFrameAttribute> FitUpDirectionWeights);
 };
 
 /** A simple class that has a single implementable function which constructs a FAnimDatabaseFrameAttribute object from a UAnimDatabase. */
@@ -2012,20 +2070,20 @@ public:
 	TArray<TObjectPtr<UAnimDatabaseFrameAttributeFunction>> FitUpDirections;
 
 	/** Optional array of weights to use for the fit locations */
-	UPROPERTY(EditAnywhere, Category = "Settings")
-	TArray<float> FitLocationWeights;
+	UPROPERTY(EditAnywhere, Instanced, Category = "Settings")
+	TArray<TObjectPtr<UAnimDatabaseFrameAttributeFunction>> FitLocationWeights;
 
 	/** Optional array of weights to use for the fit forward directions */
-	UPROPERTY(EditAnywhere, Category = "Settings")
-	TArray<float> FitForwardDirectionWeights;
+	UPROPERTY(EditAnywhere, Instanced, Category = "Settings")
+	TArray<TObjectPtr<UAnimDatabaseFrameAttributeFunction>>FitForwardDirectionWeights;
 
 	/** Optional array of weights to use for the fit right directions */
-	UPROPERTY(EditAnywhere, Category = "Settings")
-	TArray<float> FitRightDirectionWeights;
+	UPROPERTY(EditAnywhere, Instanced, Category = "Settings")
+	TArray<TObjectPtr<UAnimDatabaseFrameAttributeFunction>> FitRightDirectionWeights;
 
 	/** Optional array of weights to use for the fit up directions */
-	UPROPERTY(EditAnywhere, Category = "Settings")
-	TArray<float> FitUpDirectionWeights;
+	UPROPERTY(EditAnywhere, Instanced, Category = "Settings")
+	TArray<TObjectPtr<UAnimDatabaseFrameAttributeFunction>> FitUpDirectionWeights;
 
 public:
 
diff --git a/Engine/Plugins/Experimental/Animation/AnimDatabase/Source/AnimDatabase/Public/AnimDatabaseFrameRanges.h b/Engine/Plugins/Experimental/Animation/AnimDatabase/Source/AnimDatabase/Public/AnimDatabaseFrameRanges.h
--- a/Engine/Plugins/Experimental/Animation/AnimDatabase/Source/AnimDatabase/Public/AnimDatabaseFrameRanges.h
+++ b/Engine/Plugins/Experimental/Animation/AnimDatabase/Source/AnimDatabase/Public/AnimDatabaseFrameRanges.h
@@ -336,10 +336,26 @@ public:
 	UFUNCTION(BlueprintPure, Category = "AnimDatabase")
 	static UE_API FAnimDatabaseFrameRanges MakeFrameRangesFromNotMirrored(const UAnimDatabase* Database, const FAnimDatabaseFrameRanges& FrameRanges);
 
-	/** Make a set of FrameRanges for just the mirrored animations in a a set of frame ranges */
+	/** Make a set of FrameRanges for just the mirrored animations in a set of frame ranges */
 	UFUNCTION(BlueprintPure, Category = "AnimDatabase")
 	static UE_API FAnimDatabaseFrameRanges MakeFrameRangesFromMirrored(const UAnimDatabase* Database, const FAnimDatabaseFrameRanges& FrameRanges);
 
+	/** Make a set of FrameRanges for just the animations with root motion enabled in a set of frame ranges */
+	UFUNCTION(BlueprintPure, Category = "AnimDatabase")
+	static UE_API FAnimDatabaseFrameRanges MakeFrameRangesFromRootMotionEnabled(const UAnimDatabase* Database, const FAnimDatabaseFrameRanges& FrameRanges);
+
+	/** Make a set of FrameRanges for just the animations with root motion disabled in a set of frame ranges */
+	UFUNCTION(BlueprintPure, Category = "AnimDatabase")
+	static UE_API FAnimDatabaseFrameRanges MakeFrameRangesFromRootMotionDisabled(const UAnimDatabase* Database, const FAnimDatabaseFrameRanges& FrameRanges);
+
+	/** Make a set of FrameRanges for just the animations with force root lock enabled in a set of frame ranges */
+	UFUNCTION(BlueprintPure, Category = "AnimDatabase")
+	static UE_API FAnimDatabaseFrameRanges MakeFrameRangesFromForceRootLockEnabled(const UAnimDatabase* Database, const FAnimDatabaseFrameRanges& FrameRanges);
+
+	/** Make a set of FrameRanges for just the animations with force root lock disabled in a set of frame ranges */
+	UFUNCTION(BlueprintPure, Category = "AnimDatabase")
+	static UE_API FAnimDatabaseFrameRanges MakeFrameRangesFromForceRootLockDisabled(const UAnimDatabase* Database, const FAnimDatabaseFrameRanges& FrameRanges);
+
 	/** Make a set of FrameRanges for a single sequence in a database */
 	UFUNCTION(BlueprintPure, Category = "AnimDatabase")
 	static UE_API FAnimDatabaseFrameRanges MakeFrameRangesFromSequenceIndex(const UAnimDatabase* Database, const int32 SequenceIdx);
@@ -729,6 +745,50 @@ public:
 	UE_API virtual FAnimDatabaseFrameRanges MakeFrameRanges_Implementation(const UAnimDatabase* Database, const FAnimDatabaseFrameRanges& FrameRanges) const override;
 };
 
+/** Finds all the frame ranges that have root motion enabled */
+UCLASS(BlueprintType, Blueprintable, meta = (DisplayName = "Root Motion Enabled"))
+class UAnimDatabaseFrameRangesFunction_RootMotionEnabled : public UAnimDatabaseFrameRangesFunction
+{
+	GENERATED_BODY()
+
+public:
+
+	UE_API virtual FAnimDatabaseFrameRanges MakeFrameRanges_Implementation(const UAnimDatabase* Database, const FAnimDatabaseFrameRanges& FrameRanges) const override;
+};
+
+/** Finds all the frame ranges that have root motion disabled */
+UCLASS(BlueprintType, Blueprintable, meta = (DisplayName = "Root Motion Disabled"))
+class UAnimDatabaseFrameRangesFunction_RootMotionDisabled : public UAnimDatabaseFrameRangesFunction
+{
+	GENERATED_BODY()
+
+public:
+
+	UE_API virtual FAnimDatabaseFrameRanges MakeFrameRanges_Implementation(const UAnimDatabase* Database, const FAnimDatabaseFrameRanges& FrameRanges) const override;
+};
+
+/** Finds all the frame ranges that have force root lock enabled */
+UCLASS(BlueprintType, Blueprintable, meta = (DisplayName = "Force Root Lock Enabled"))
+class UAnimDatabaseFrameRangesFunction_ForceRootLockEnabled : public UAnimDatabaseFrameRangesFunction
+{
+	GENERATED_BODY()
+
+public:
+
+	UE_API virtual FAnimDatabaseFrameRanges MakeFrameRanges_Implementation(const UAnimDatabase* Database, const FAnimDatabaseFrameRanges& FrameRanges) const override;
+};
+
+/** Finds all the frame ranges that have force root lock disabled */
+UCLASS(BlueprintType, Blueprintable, meta = (DisplayName = "Force Root Lock Disabled"))
+class UAnimDatabaseFrameRangesFunction_ForceRootLockDisabled : public UAnimDatabaseFrameRangesFunction
+{
+	GENERATED_BODY()
+
+public:
+
+	UE_API virtual FAnimDatabaseFrameRanges MakeFrameRanges_Implementation(const UAnimDatabase* Database, const FAnimDatabaseFrameRanges& FrameRanges) const override;
+};
+
 /** Finds all the frame ranges corresponding to the given Anim Sequence */
 UCLASS(BlueprintType, Blueprintable, meta = (DisplayName = "Anim Sequence"))
 class UAnimDatabaseFrameRangesFunction_AnimSequence : public UAnimDatabaseFrameRangesFunction
diff --git a/Engine/Plugins/Experimental/Animation/AnimDatabase/Source/AnimDatabaseEditor/Private/AnimDatabaseEditorToolkit.cpp b/Engine/Plugins/Experimental/Animation/AnimDatabase/Source/AnimDatabaseEditor/Private/AnimDatabaseEditorToolkit.cpp
--- a/Engine/Plugins/Experimental/Animation/AnimDatabase/Source/AnimDatabaseEditor/Private/AnimDatabaseEditorToolkit.cpp
+++ b/Engine/Plugins/Experimental/Animation/AnimDatabase/Source/AnimDatabaseEditor/Private/AnimDatabaseEditorToolkit.cpp
@@ -131,10 +131,6 @@ namespace UE::AnimDatabase::Editor
 		UAnimDatabaseQuery* Query = StaticCast<UAnimDatabaseQuery*>(Objects[0].Pin().Get());
 		if (!Query) { return; }
 
-		// Reset the selected items since we are re-drawing the UI
-
-		SelectedItems.Reset();
-
 		// Put the Query category first
 
 		IDetailCategoryBuilder& QueryCategory = DetailBuilder.EditCategory("Query", LOCTEXT("QueryCategory", "Query"), ECategoryPriority::Important);
@@ -173,10 +169,6 @@ namespace UE::AnimDatabase::Editor
 					.FillWidth(0.05f)
 				);
 
-			// Select first entry if we are re-drawing the UI
-
-			ListView->SetSelection(Query->QueryEntries[0]);
-
 			// Add to UI
 
 			RangesCategory.AddCustomRow(LOCTEXT("RangesProperty", "Ranges"))
@@ -191,6 +183,18 @@ namespace UE::AnimDatabase::Editor
 							ListView.ToSharedRef()
 						]
 				];
+
+			// Re-set selection
+
+			TArray<TSharedPtr<FQueryEntry>, TInlineAllocator<16>> ResetSelectedItems;
+			for (const TSharedPtr<FQueryEntry>& Item : ListView->GetItems())
+			{
+				if (Item->bIsSelected)
+				{
+					ResetSelectedItems.Add(Item);
+				}
+			}
+			ListView->SetItemSelection(ResetSelectedItems, true);
 		}
 		else
 		{
@@ -208,8 +212,6 @@ namespace UE::AnimDatabase::Editor
 		IDetailCategoryBuilder& TimelineCategory = DetailBuilder.EditCategory("Timeline", LOCTEXT("TimelineCategory", "Timeline"), ECategoryPriority::Important);
 		IDetailCategoryBuilder& DebugDrawCategory = DetailBuilder.EditCategory("Debug Draw", LOCTEXT("DebugDrawCategory", "Debug Draw"), ECategoryPriority::Important);
 		IDetailCategoryBuilder& FunctionCategory = DetailBuilder.EditCategory("Function", LOCTEXT("FunctionCategory", "Function"), ECategoryPriority::Important);
-
-
 	}
 
 	FReply FQueryDetails::OnKeyDown(UAnimDatabaseQuery* Query, const FKeyEvent& InKeyEvent)
diff --git a/Engine/Plugins/Experimental/Animation/AnimGen/Source/AnimGen/Private/AnimGenControl.cpp b/Engine/Plugins/Experimental/Animation/AnimGen/Source/AnimGen/Private/AnimGenControl.cpp
--- a/Engine/Plugins/Experimental/Animation/AnimGen/Source/AnimGen/Private/AnimGenControl.cpp
+++ b/Engine/Plugins/Experimental/Animation/AnimGen/Source/AnimGen/Private/AnimGenControl.cpp
@@ -2355,7 +2355,7 @@ FAnimGenControlObjectElement UAnimGenControls::MakeTrajectoryControlFromTransfor
 {
 	if (SampleNum < 0)
 	{
-		UE_LOGFMT(LogAnimDatabase, Error, "MakeTrajectoryControlFromTransformTrajectory: Invalid Sample Num: {SampleNum}.", SampleNum);
+		UE_ANIMGEN_LOG_ONCE(LogAnimDatabase, Error, "MakeTrajectoryControlFromTransformTrajectory: Invalid Sample Num: {SampleNum}.", SampleNum);
 		return FAnimGenControlObjectElement();
 	}
 
@@ -2379,7 +2379,7 @@ FAnimGenControlObjectElement UAnimGenControls::MakeLocationTrajectoryControlFrom
 {
 	if (SampleNum < 0)
 	{
-		UE_LOGFMT(LogAnimDatabase, Error, "MakeLocationTrajectoryControlFromTransformTrajectory: Invalid Sample Num: {SampleNum}.", SampleNum);
+		UE_ANIMGEN_LOG_ONCE(LogAnimDatabase, Error, "MakeLocationTrajectoryControlFromTransformTrajectory: Invalid Sample Num: {SampleNum}.", SampleNum);
 		return FAnimGenControlObjectElement();
 	}
 
diff --git a/Engine/Plugins/Experimental/Animation/AnimGen/Source/AnimGen/Public/AnimGenAutoEncoder.h b/Engine/Plugins/Experimental/Animation/AnimGen/Source/AnimGen/Public/AnimGenAutoEncoder.h
--- a/Engine/Plugins/Experimental/Animation/AnimGen/Source/AnimGen/Public/AnimGenAutoEncoder.h
+++ b/Engine/Plugins/Experimental/Animation/AnimGen/Source/AnimGen/Public/AnimGenAutoEncoder.h
@@ -625,11 +625,11 @@ public:
 public:
 
 	/** Encoder Network */
-	UPROPERTY()
+	UPROPERTY(Instanced)
 	TObjectPtr<ULearningNeuralNetworkData> EncoderNetwork = nullptr;
 
 	/** Decoder Network */
-	UPROPERTY()
+	UPROPERTY(Instanced)
 	TObjectPtr<ULearningNeuralNetworkData> DecoderNetwork = nullptr;
 
 public:
@@ -747,11 +747,11 @@ public:
 public:
 
 	/** Editor-only object containing the viewport settings for the editor window */
-	UPROPERTY()
+	UPROPERTY(Instanced)
 	TObjectPtr<UAnimGenAutoEncoderViewportSettings> ViewportSettings;
 
 	/** Editor-only object containing the training settings for the editor window */
-	UPROPERTY()
+	UPROPERTY(Instanced)
 	TObjectPtr<UAnimGenAutoEncoderTrainingSettings> TrainingSettings;
 
 #endif
diff --git a/Engine/Plugins/Experimental/Animation/AnimGen/Source/AnimGen/Public/AnimGenController.h b/Engine/Plugins/Experimental/Animation/AnimGen/Source/AnimGen/Public/AnimGenController.h
--- a/Engine/Plugins/Experimental/Animation/AnimGen/Source/AnimGen/Public/AnimGenController.h
+++ b/Engine/Plugins/Experimental/Animation/AnimGen/Source/AnimGen/Public/AnimGenController.h
@@ -584,19 +584,19 @@ public:
 	FAnimGenControlSchemaElement ControlSchemaElement;
 
 	/** Control Encoder Network */
-	UPROPERTY()
+	UPROPERTY(Instanced)
 	TObjectPtr<ULearningNeuralNetworkData> ControlEncoderNetwork = nullptr;
 
 	/** LOD0 Network */
-	UPROPERTY()
+	UPROPERTY(Instanced)
 	TObjectPtr<ULearningNeuralNetworkData> LOD0Network = nullptr;
 
 	/** LOD1 Network */
-	UPROPERTY()
+	UPROPERTY(Instanced)
 	TObjectPtr<ULearningNeuralNetworkData> LOD1Network = nullptr;
 
 	/** LOD2 Network */
-	UPROPERTY()
+	UPROPERTY(Instanced)
 	TObjectPtr<ULearningNeuralNetworkData> LOD2Network = nullptr;
 
 	/** The Frame Rate this controller was trained to run at */
@@ -684,15 +684,15 @@ public:
 public:
 
 	/** Editor-only object containing the training settings for the editor window */
-	UPROPERTY()
+	UPROPERTY(Instanced)
 	TObjectPtr<UAnimGenControllerTrainingSettings> TrainingSettings;
 
 	/** Editor-only object containing the testing settings for the editor window */
-	UPROPERTY()
+	UPROPERTY(Instanced)
 	TObjectPtr<UAnimGenControllerTestingSettings> TestingSettings;
 
 	/** Editor-only object containing the viewport settings for the editor window */
-	UPROPERTY()
+	UPROPERTY(Instanced)
 	TObjectPtr<UAnimGenControllerViewportSettings> ViewportSettings;
 
 #endif
diff --git a/Engine/Plugins/Experimental/Animation/AnimGen/Source/AnimGenEditor/Private/AnimGenEditorAutoEncoderToolkit.cpp b/Engine/Plugins/Experimental/Animation/AnimGen/Source/AnimGenEditor/Private/AnimGenEditorAutoEncoderToolkit.cpp
--- a/Engine/Plugins/Experimental/Animation/AnimGen/Source/AnimGenEditor/Private/AnimGenEditorAutoEncoderToolkit.cpp
+++ b/Engine/Plugins/Experimental/Animation/AnimGen/Source/AnimGenEditor/Private/AnimGenEditorAutoEncoderToolkit.cpp
@@ -239,6 +239,18 @@ namespace UE::AnimGen::Editor
 							ListView.ToSharedRef()
 						]
 				];
+
+			// Re-set selection
+
+			TArray<TSharedPtr<AnimDatabase::Editor::FQueryEntry>, TInlineAllocator<16>> ResetSelectedItems;
+			for (const TSharedPtr<AnimDatabase::Editor::FQueryEntry>& Item : ListView->GetItems())
+			{
+				if (Item->bIsSelected)
+				{
+					ResetSelectedItems.Add(Item);
+				}
+			}
+			ListView->SetItemSelection(ResetSelectedItems, true);
 		}
 		else
 		{
diff --git a/Engine/Plugins/Experimental/Animation/AnimGen/Source/AnimGenEditor/Private/AnimGenEditorControllerToolkit.cpp b/Engine/Plugins/Experimental/Animation/AnimGen/Source/AnimGenEditor/Private/AnimGenEditorControllerToolkit.cpp
--- a/Engine/Plugins/Experimental/Animation/AnimGen/Source/AnimGenEditor/Private/AnimGenEditorControllerToolkit.cpp
+++ b/Engine/Plugins/Experimental/Animation/AnimGen/Source/AnimGenEditor/Private/AnimGenEditorControllerToolkit.cpp
@@ -237,6 +237,18 @@ namespace UE::AnimGen::Editor
 						ListView.ToSharedRef()
 					]
 				];
+
+			// Re-set selection
+
+			TArray<TSharedPtr<AnimDatabase::Editor::FQueryEntry>, TInlineAllocator<16>> ResetSelectedItems;
+			for (const TSharedPtr<AnimDatabase::Editor::FQueryEntry>& Item : ListView->GetItems())
+			{
+				if (Item->bIsSelected)
+				{
+					ResetSelectedItems.Add(Item);
+				}
+			}
+			ListView->SetItemSelection(ResetSelectedItems, true);
 		}
 		else
 		{
