AnimDatabase: Minor improvements for Frame Attributes

#Synced-CL 51630663

diff --git a/Engine/Plugins/Experimental/Animation/AnimDatabase/Source/AnimDatabase/Private/AnimDatabase.cpp b/Engine/Plugins/Experimental/Animation/AnimDatabase/Source/AnimDatabase/Private/AnimDatabase.cpp
--- a/Engine/Plugins/Experimental/Animation/AnimDatabase/Source/AnimDatabase/Private/AnimDatabase.cpp
+++ b/Engine/Plugins/Experimental/Animation/AnimDatabase/Source/AnimDatabase/Private/AnimDatabase.cpp
@@ -6400,7 +6400,7 @@ void UAnimDatabaseLibrary::DatabaseSetLocalBoneTransformsFromArrayViews(UAnimDat
 				ScaleKeys.Add(FrameTransform.GetScale3D());
 			}
 
-			FInt32Range KeyRange(StartFrame, FMath::Min(StartFrame + SequenceFrameNum, SequenceTotalFrameNum));
+			FInt32Range KeyRange(SequenceStartFrame, FMath::Min(SequenceStartFrame + SequenceFrameNum, SequenceTotalFrameNum));
 			if (!Controller->UpdateBoneTrackKeys(BoneNames[BoneIdx], KeyRange, LocationKeys, RotationKeys, ScaleKeys, bShouldTransact))
 			{
 				UE_LOGFMT(LogAnimDatabase, Error, "DatabaseSetLocalBoneTransformsFromArrayViews: Updating Transform Keys Failed.");
@@ -6454,6 +6454,116 @@ void UAnimDatabaseLibrary::DatabaseSetGlobalBoneTransformsFromArrayViews(UAnimDa
 #endif
 }
 
+void UAnimDatabaseLibrary::DatabaseResetBoneTransform(UAnimDatabase* Database, const FAnimDatabaseFrameRanges& FrameRanges, const FName BoneName, const bool bShouldTransact)
+{
+	DatabaseResetBoneTransformsFromArrayViews(Database, FrameRanges, { BoneName }, bShouldTransact);
+}
+
+void UAnimDatabaseLibrary::DatabaseResetBoneTransforms(UAnimDatabase* Database, const FAnimDatabaseFrameRanges& FrameRanges, const TArray<FName>& BoneNames, const bool bShouldTransact)
+{
+	DatabaseResetBoneTransformsFromArrayViews(Database, FrameRanges, BoneNames, bShouldTransact);
+}
+
+void UAnimDatabaseLibrary::DatabaseResetBoneTransformsFromArrayViews(UAnimDatabase* Database, const FAnimDatabaseFrameRanges& FrameRanges, const TArrayView<const FName> BoneNames, const bool bShouldTransact)
+{
+#if WITH_EDITOR
+	TRACE_CPUPROFILER_EVENT_SCOPE(UAnimDatabaseLibrary::DatabaseResetBoneTransformsFromArrayViews);
+
+	const int32 BoneNum = BoneNames.Num();
+
+	const FFrameRate DatabaseFrameRate = Database->GetFrameRate();
+
+	const FAnimDatabaseFrameRanges FilteredFrameRanges = UAnimDatabaseFrameRangesLibrary::MakeFrameRangesFromNotMirrored(Database, FrameRanges);
+
+	// To avoid triggering recompression during modification we need to open brackets on all the controllers for all the to-be-modified animation 
+	// sequences once, and then close them all at the end.
+
+	TArray<TScriptInterface<IAnimationDataController>> Controllers;
+	TMap<int32, int32> SequenceControllerIndices;
+
+	UE::AnimDatabase::Private::OpenControllers(
+		Controllers, SequenceControllerIndices, Database, FilteredFrameRanges,
+		LOCTEXT("ResetBoneTransformScript", "Resetting Bone Transform"), bShouldTransact);
+
+	FScopedSlowTask SlowTask(BoneNum, LOCTEXT("ResettingBoneLocalTransforms", "Resetting Bone Transforms"));
+	SlowTask.MakeDialog(true);
+
+	for (int32 BoneIdx = 0; BoneIdx < BoneNum; BoneIdx++)
+	{
+		if (SlowTask.ShouldCancel()) { break; }
+		SlowTask.EnterProgressFrame();
+
+		const int32 BoneIndex = Database->FindBoneIndex(BoneNames[BoneIdx]);
+
+		if (BoneIndex == INDEX_NONE)
+		{
+			UE_LOGFMT(LogAnimDatabase, Error, "DatabaseResetBoneTransformsFromArrayViews: Could not find bone {BoneName}.", *BoneNames[BoneIdx].ToString());
+			continue;
+		}
+
+		const int32 TotalRangeNum = FilteredFrameRanges.FrameRangeSet->GetTotalRangeNum();
+
+		TLearningArray<1, int32> TagRangeSequenceIndices;
+		TagRangeSequenceIndices.SetNumUninitialized({ TotalRangeNum });
+		UE::Learning::FrameRangeSet::AllRangeSequences(TagRangeSequenceIndices, *FilteredFrameRanges.FrameRangeSet);
+
+		TArray<FVector3f> LocationKeys;
+		TArray<FQuat4f> RotationKeys;
+		TArray<FVector3f> ScaleKeys;
+
+		for (int32 RangeIdx = 0; RangeIdx < TotalRangeNum; RangeIdx++)
+		{
+			TRACE_CPUPROFILER_EVENT_SCOPE(UAnimDatabaseLibrary::DatabaseResetBoneTransformsFromArrayViews::Range);
+
+			const int32 RangeSequenceIdx = TagRangeSequenceIndices[RangeIdx];
+
+			check(!Database->GetIsMirrored(RangeSequenceIdx));
+
+			TScriptInterface<IAnimationDataController> Controller = Controllers[SequenceControllerIndices[RangeSequenceIdx]];
+
+			// Skip bone track if it does not exist
+			if (!Controller->GetModel()->IsValidBoneTrackName(BoneNames[BoneIdx]))
+			{
+				continue;
+			}
+
+			// Set Transform Keys
+			const int32 RangeFrameNum = FilteredFrameRanges.FrameRangeSet->GetAllRangeLengths()[RangeIdx];
+			const int32 FrameAttributeRangeOffset = FilteredFrameRanges.FrameRangeSet->GetAllRangeOffsets()[RangeIdx];
+			const int32 StartFrame = FilteredFrameRanges.FrameRangeSet->GetAllRangeStarts()[RangeIdx];
+
+			const UAnimSequence* RangeAnimSequence = Database->GetAnimSequence(RangeSequenceIdx);
+			const FFrameRate SequenceFrameRate = RangeAnimSequence->GetSamplingFrameRate();
+			const int32 SequenceFrameNum = (SequenceFrameRate.Numerator * RangeFrameNum * DatabaseFrameRate.Denominator) /
+				(SequenceFrameRate.Denominator * DatabaseFrameRate.Numerator);
+			const int32 SequenceStartFrame = (SequenceFrameRate.Numerator * StartFrame * DatabaseFrameRate.Denominator) /
+				(SequenceFrameRate.Denominator * DatabaseFrameRate.Numerator);
+			const int32 SequenceTotalFrameNum = RangeAnimSequence->GetNumberOfSampledKeys();
+
+			LocationKeys.Reset(SequenceFrameNum);
+			RotationKeys.Reset(SequenceFrameNum);
+			ScaleKeys.Reset(SequenceFrameNum);
+
+			for (int32 FrameIdx = 0; FrameIdx < SequenceFrameNum; FrameIdx++)
+			{
+				const FTransform FrameTransform = Database->GetSkeleton()->GetReferenceSkeleton().GetRefBonePose()[BoneIndex];
+				LocationKeys.Add((FVector3f)FrameTransform.GetLocation());
+				RotationKeys.Add((FQuat4f)FrameTransform.GetRotation());
+				ScaleKeys.Add((FVector3f)FrameTransform.GetScale3D());
+			}
+
+			FInt32Range KeyRange(SequenceStartFrame, FMath::Min(SequenceStartFrame + SequenceFrameNum, SequenceTotalFrameNum));
+			if (!Controller->UpdateBoneTrackKeys(BoneNames[BoneIdx], KeyRange, LocationKeys, RotationKeys, ScaleKeys, bShouldTransact))
+			{
+				UE_LOGFMT(LogAnimDatabase, Error, "DatabaseResetBoneTransformsFromArrayViews: Updating Transform Keys Failed.");
+			}
+		}
+	}
+
+	UE::AnimDatabase::Private::CloseControllers(Controllers, bShouldTransact);
+#endif
+}
+
 void UAnimDatabaseLibrary::DatabaseStripUnusedBoneTracks(UAnimDatabase* Database, const FAnimDatabaseFrameRanges& FrameRanges, const bool bShouldTransact)
 {
 #if WITH_EDITOR
@@ -7045,6 +7155,46 @@ void UAnimDatabaseFunction_SetBoneGlobalTransformFromFrameAttribute::Run_Impleme
 			UAnimDatabaseFrameAttributeLibrary::MakeBoneGlobalTransformFrameAttribute(InDatabase, InFrameRanges, ParentIndex)), bShouldTransact);
 }
 
+void UAnimDatabaseFunction_CopyBones::Run_Implementation(UAnimDatabase* InDatabase, const FAnimDatabaseFrameRanges& InFrameRanges, const bool bShouldTransact)
+{
+	const int32 BoneNum = BoneMap.Num();
+
+	TArray<FAnimDatabaseFrameAttribute, TInlineAllocator<16>> FrameAttributes;
+	FrameAttributes.SetNum(BoneNum);
+
+	TArray<FName, TInlineAllocator<16>> SrcBones, DstBones;
+	SrcBones.SetNum(BoneNum);
+	DstBones.SetNum(BoneNum);
+
+	BoneMap.GetKeys(SrcBones);
+	for (int32 BoneIdx = 0; BoneIdx < BoneNum; BoneIdx++)
+	{
+		DstBones[BoneIdx] = BoneMap[SrcBones[BoneIdx]];
+	}
+
+	switch (Space)
+	{
+	case EAnimDatabaseBoneSpace::Global:
+	{
+		UAnimDatabaseFrameAttributeLibrary::MakeBoneGlobalTransformFrameAttributesFromNamesArrayView(FrameAttributes, InDatabase, InFrameRanges, SrcBones);
+		UAnimDatabaseLibrary::DatabaseSetGlobalBoneTransformsFromArrayViews(InDatabase, DstBones, FrameAttributes, bShouldTransact);
+		return;
+	}
+	
+	case EAnimDatabaseBoneSpace::Local:
+	{
+		UAnimDatabaseFrameAttributeLibrary::MakeBoneLocalTransformFrameAttributesFromNamesArrayView(FrameAttributes, InDatabase, InFrameRanges, SrcBones);
+		UAnimDatabaseLibrary::DatabaseSetLocalBoneTransformsFromArrayViews(InDatabase, DstBones, FrameAttributes, bShouldTransact);
+		return;
+	}
+	}
+}
+
+void UAnimDatabaseFunction_ResetBones::Run_Implementation(UAnimDatabase* InDatabase, const FAnimDatabaseFrameRanges& InFrameRanges, const bool bShouldTransact)
+{
+	UAnimDatabaseLibrary::DatabaseResetBoneTransformsFromArrayViews(InDatabase, InFrameRanges, Bones, bShouldTransact);
+}
+
 void UAnimDatabaseFunction_SetBoneToRootTransformAtFrames::Run_Implementation(UAnimDatabase* InDatabase, const FAnimDatabaseFrameRanges& InFrameRanges, const bool bShouldTransact)
 {
 	const int32 BoneIndex = InDatabase->FindBoneIndex(BoneName);
diff --git a/Engine/Plugins/Experimental/Animation/AnimDatabase/Source/AnimDatabase/Private/AnimDatabaseFrameAttribute.cpp b/Engine/Plugins/Experimental/Animation/AnimDatabase/Source/AnimDatabase/Private/AnimDatabaseFrameAttribute.cpp
--- a/Engine/Plugins/Experimental/Animation/AnimDatabase/Source/AnimDatabase/Private/AnimDatabaseFrameAttribute.cpp
+++ b/Engine/Plugins/Experimental/Animation/AnimDatabase/Source/AnimDatabase/Private/AnimDatabaseFrameAttribute.cpp
@@ -598,7 +598,7 @@ void UAnimDatabaseFrameAttributeLibrary::MakeBoolFrameAttributesFromActiveRanges
 	MakeBoolFrameAttributesFromActiveRangesArrayViews(OutFrameAttributes, Actives, FrameRanges);
 }
 
-void UAnimDatabaseFrameAttributeLibrary::MakeBoolFrameAttributesFromActiveRangesArrayViews(TArrayView<FAnimDatabaseFrameAttribute> OutFrameAttributes, const TArrayView<const FAnimDatabaseFrameRanges> Actives, const FAnimDatabaseFrameRanges& FrameRanges)
+void UAnimDatabaseFrameAttributeLibrary::MakeBoolFrameAttributesFromActiveRangesArrayViews(const TArrayView<FAnimDatabaseFrameAttribute> OutFrameAttributes, const TArrayView<const FAnimDatabaseFrameRanges> Actives, const FAnimDatabaseFrameRanges& FrameRanges)
 {
 	check(OutFrameAttributes.Num() == Actives.Num());
 
@@ -684,7 +684,7 @@ void UAnimDatabaseFrameAttributeLibrary::MakeBoolFrameAttributesFromCurvesActive
 	MakeBoolFrameAttributesFromCurvesActiveArrayView(OutBoolFrameAttributes, Database, FrameRanges, CurveNames);
 }
 
-void UAnimDatabaseFrameAttributeLibrary::MakeBoolFrameAttributesFromCurvesActiveArrayView(TArrayView<FAnimDatabaseFrameAttribute> OutBoolFrameAttributes, const UAnimDatabase* Database, const FAnimDatabaseFrameRanges& FrameRanges, const TArrayView<const FName> CurveNames)
+void UAnimDatabaseFrameAttributeLibrary::MakeBoolFrameAttributesFromCurvesActiveArrayView(const TArrayView<FAnimDatabaseFrameAttribute> OutBoolFrameAttributes, const UAnimDatabase* Database, const FAnimDatabaseFrameRanges& FrameRanges, const TArrayView<const FName> CurveNames)
 {
 	TRACE_CPUPROFILER_EVENT_SCOPE(UAnimDatabaseFrameAttributeLibrary::MakeBoolFrameAttributesFromCurvesActiveArrayView);
 
@@ -964,7 +964,7 @@ void UAnimDatabaseFrameAttributeLibrary::MakeFloatFrameAttributesFromCurves(TArr
 	MakeFloatFrameAttributesFromCurvesArrayView(OutFloatFrameAttributes, Database, FrameRanges, CurveNames);
 }
 
-void UAnimDatabaseFrameAttributeLibrary::MakeFloatFrameAttributesFromCurvesArrayView(TArrayView<FAnimDatabaseFrameAttribute> OutFloatFrameAttributes, const UAnimDatabase* Database, const FAnimDatabaseFrameRanges& FrameRanges, const TArrayView<const FName> CurveNames)
+void UAnimDatabaseFrameAttributeLibrary::MakeFloatFrameAttributesFromCurvesArrayView(const TArrayView<FAnimDatabaseFrameAttribute> OutFloatFrameAttributes, const UAnimDatabase* Database, const FAnimDatabaseFrameRanges& FrameRanges, const TArrayView<const FName> CurveNames)
 {
 	TRACE_CPUPROFILER_EVENT_SCOPE(UAnimDatabaseFrameAttributeLibrary::MakeFloatFrameAttributeFromCurve);
 
@@ -1050,7 +1050,7 @@ void UAnimDatabaseFrameAttributeLibrary::MakeFloatFrameAttributesFromCurvesWhenA
 	MakeFloatFrameAttributesFromCurvesWhenActiveArrayView(OutFloatFrameAttributes, Database, FrameRanges, CurveNames);
 }
 
-void UAnimDatabaseFrameAttributeLibrary::MakeFloatFrameAttributesFromCurvesWhenActiveArrayView(TArrayView<FAnimDatabaseFrameAttribute> OutFloatFrameAttributes, const UAnimDatabase* Database, const FAnimDatabaseFrameRanges& FrameRanges, const TArrayView<const FName> CurveNames)
+void UAnimDatabaseFrameAttributeLibrary::MakeFloatFrameAttributesFromCurvesWhenActiveArrayView(const TArrayView<FAnimDatabaseFrameAttribute> OutFloatFrameAttributes, const UAnimDatabase* Database, const FAnimDatabaseFrameRanges& FrameRanges, const TArrayView<const FName> CurveNames)
 {
 	const int32 CurveNum = CurveNames.Num();
 	check(OutFloatFrameAttributes.Num() == CurveNum);
@@ -2354,7 +2354,7 @@ void UAnimDatabaseFrameAttributeLibrary::MakeBoneLocalTransformFrameAttributes(
 
 	OutTransformFrameAttributes.SetNum(BoneIndices.Num());
 
-	MakeBoneLocalFrameAttributesFromArrayViews(
+	MakeBoneLocalTransformFrameAttributesFromArrayViews(
 		OutTransformFrameAttributes,
 		Database,
 		FrameRanges,
@@ -2362,7 +2362,47 @@ void UAnimDatabaseFrameAttributeLibrary::MakeBoneLocalTransformFrameAttributes(
 		RelativeTime);
 }
 
-void UAnimDatabaseFrameAttributeLibrary::MakeBoneLocalFrameAttributesFromArrayViews(
+void UAnimDatabaseFrameAttributeLibrary::MakeBoneLocalTransformFrameAttributesFromNames(
+	TArray<FAnimDatabaseFrameAttribute>& OutTransformFrameAttributes,
+	const UAnimDatabase* Database,
+	const FAnimDatabaseFrameRanges& FrameRanges,
+	const TArray<FName>& BoneNames,
+	const float RelativeTime)
+{
+	OutTransformFrameAttributes.SetNum(BoneNames.Num());
+	MakeBoneLocalTransformFrameAttributesFromNamesArrayView(OutTransformFrameAttributes, Database, FrameRanges, BoneNames, RelativeTime);
+}
+
+void UAnimDatabaseFrameAttributeLibrary::MakeBoneLocalTransformFrameAttributesFromNamesArrayView(
+	const TArrayView<FAnimDatabaseFrameAttribute> OutTransformFrameAttributes,
+	const UAnimDatabase* Database,
+	const FAnimDatabaseFrameRanges& FrameRanges,
+	const TArrayView<const FName> BoneNames,
+	const float RelativeTime)
+{
+	if (!Database)
+	{
+		UE_LOGFMT(LogAnimDatabase, Error, "MakeBoneLocalTransformFrameAttributesFromNamesArrayView: Database is nullptr.");
+		UE::Learning::Array::Set<1, FAnimDatabaseFrameAttribute>(OutTransformFrameAttributes, FAnimDatabaseFrameAttribute());
+		return;
+	}
+
+	check(OutTransformFrameAttributes.Num() == BoneNames.Num());
+
+	TArray<int32, TInlineAllocator<32>> BoneIndices;
+	const int32 BoneNum = BoneNames.Num();
+	BoneIndices.SetNumUninitialized(BoneNum);
+	Database->FindBoneIndicesFromArrayViews(BoneIndices, BoneNames);
+
+	MakeBoneLocalTransformFrameAttributesFromArrayViews(
+		OutTransformFrameAttributes,
+		Database,
+		FrameRanges,
+		BoneIndices,
+		RelativeTime);
+}
+
+void UAnimDatabaseFrameAttributeLibrary::MakeBoneLocalTransformFrameAttributesFromArrayViews(
 	const TArrayView<FAnimDatabaseFrameAttribute> OutTransformFrameAttributes,
 	const UAnimDatabase* Database,
 	const FAnimDatabaseFrameRanges& FrameRanges,
@@ -2371,14 +2411,14 @@ void UAnimDatabaseFrameAttributeLibrary::MakeBoneLocalFrameAttributesFromArrayVi
 {
 	if (!Database)
 	{
-		UE_LOGFMT(LogAnimDatabase, Error, "MakeBoneLocalFrameAttributesFromArrayViews: Database is nullptr.");
+		UE_LOGFMT(LogAnimDatabase, Error, "MakeBoneLocalTransformFrameAttributesFromArrayViews: Database is nullptr.");
 		UE::Learning::Array::Set<1, FAnimDatabaseFrameAttribute>(OutTransformFrameAttributes, FAnimDatabaseFrameAttribute());
 		return;
 	}
 
 	if (!FrameRanges.IsValid())
 	{
-		UE_LOGFMT(LogAnimDatabase, Error, "MakeBoneLocalFrameAttributesFromArrayViews: Invalid FrameRanges.");
+		UE_LOGFMT(LogAnimDatabase, Error, "MakeBoneLocalTransformFrameAttributesFromArrayViews: Invalid FrameRanges.");
 		UE::Learning::Array::Set<1, FAnimDatabaseFrameAttribute>(OutTransformFrameAttributes, FAnimDatabaseFrameAttribute());
 		return;
 	}
@@ -2390,6 +2430,13 @@ void UAnimDatabaseFrameAttributeLibrary::MakeBoneLocalFrameAttributesFromArrayVi
 
 	for (int32 BoneIdx = 0; BoneIdx < BoneIndexNum; BoneIdx++)
 	{
+		if (BoneIndices[BoneIdx] == INDEX_NONE)
+		{
+			UE_LOGFMT(LogAnimDatabase, Error, "MakeBoneLocalTransformFrameAttributesFromArrayViews: Invalid Bone Index.");
+			UE::Learning::Array::Set<1, FAnimDatabaseFrameAttribute>(OutTransformFrameAttributes, FAnimDatabaseFrameAttribute());
+			return;
+		}
+
 		OutTransformFrameAttributes[BoneIdx] = MakeEmptyFrameAttribute();
 		OutTransformFrameAttributes[BoneIdx].Type = EAnimDatabaseAttributeType::Transform;
 		OutTransformFrameAttributes[BoneIdx].FrameAttribute->FrameRangeSet = *FrameRanges.FrameRangeSet;
@@ -2458,7 +2505,7 @@ void UAnimDatabaseFrameAttributeLibrary::MakeBoneLocalFrameAttributesFromArrayVi
 FAnimDatabaseFrameAttribute UAnimDatabaseFrameAttributeLibrary::MakeBoneLocalTransformFrameAttribute(const UAnimDatabase* Database, const FAnimDatabaseFrameRanges& FrameRanges, const int32 BoneIndex, const float RelativeTime)
 {
 	FAnimDatabaseFrameAttribute Out;
-	MakeBoneLocalFrameAttributesFromArrayViews(MakeArrayView(&Out, 1), Database, FrameRanges, { BoneIndex }, RelativeTime);
+	MakeBoneLocalTransformFrameAttributesFromArrayViews(MakeArrayView(&Out, 1), Database, FrameRanges, { BoneIndex }, RelativeTime);
 	return Out;
 }
 
@@ -3151,7 +3198,7 @@ void UAnimDatabaseFrameAttributeLibrary::MakeBoneGlobalTransformFrameAttributesF
 }
 
 void UAnimDatabaseFrameAttributeLibrary::MakeBoneGlobalTransformFrameAttributesFromNamesArrayView(
-	TArrayView<FAnimDatabaseFrameAttribute> OutTransformFrameAttributes,
+	const TArrayView<FAnimDatabaseFrameAttribute> OutTransformFrameAttributes,
 	const UAnimDatabase* Database,
 	const FAnimDatabaseFrameRanges& FrameRanges,
 	const TArrayView<const FName> BoneNames,
@@ -3216,7 +3263,7 @@ void UAnimDatabaseFrameAttributeLibrary::MakeBoneGlobalLocationFrameAttributes(T
 	MakeBoneGlobalLocationFrameAttributesFromArrayViews(OutLocationFrameAttributes, Database, FrameRanges, BoneIndices, RelativeTime);
 }
 
-void UAnimDatabaseFrameAttributeLibrary::MakeBoneGlobalLocationFrameAttributesFromArrayViews(TArrayView<FAnimDatabaseFrameAttribute> OutLocationFrameAttributes, const UAnimDatabase* Database, const FAnimDatabaseFrameRanges& FrameRanges, const TArrayView<const int32> BoneIndices, const float RelativeTime)
+void UAnimDatabaseFrameAttributeLibrary::MakeBoneGlobalLocationFrameAttributesFromArrayViews(const TArrayView<FAnimDatabaseFrameAttribute> OutLocationFrameAttributes, const UAnimDatabase* Database, const FAnimDatabaseFrameRanges& FrameRanges, const TArrayView<const int32> BoneIndices, const float RelativeTime)
 {
 	MakeBoneGlobalFrameAttributesFromArrayViews(
 		OutLocationFrameAttributes, {}, {},
@@ -3234,7 +3281,7 @@ void UAnimDatabaseFrameAttributeLibrary::MakeBoneGlobalLocationFrameAttributesFr
 	MakeBoneGlobalLocationFrameAttributesFromNamesArrayView(OutLocationFrameAttributes, Database, FrameRanges, BoneNames, RelativeTime);
 }
 
-void UAnimDatabaseFrameAttributeLibrary::MakeBoneGlobalLocationFrameAttributesFromNamesArrayView(TArrayView<FAnimDatabaseFrameAttribute> OutLocationFrameAttributes, const UAnimDatabase* Database, const FAnimDatabaseFrameRanges& FrameRanges, const TArrayView<const FName> BoneNames, const float RelativeTime)
+void UAnimDatabaseFrameAttributeLibrary::MakeBoneGlobalLocationFrameAttributesFromNamesArrayView(const TArrayView<FAnimDatabaseFrameAttribute> OutLocationFrameAttributes, const UAnimDatabase* Database, const FAnimDatabaseFrameRanges& FrameRanges, const TArrayView<const FName> BoneNames, const float RelativeTime)
 {
 	if (!Database)
 	{
@@ -3327,7 +3374,7 @@ void UAnimDatabaseFrameAttributeLibrary::MakeBoneGlobalDirectionFrameAttributesF
 	MakeBoneGlobalDirectionFrameAttributesFromNamesArrayView(OutDirectionFrameAttributes, Database, FrameRanges, BoneNames, LocalDirection, RelativeTime);
 }
 
-void UAnimDatabaseFrameAttributeLibrary::MakeBoneGlobalDirectionFrameAttributesFromNamesArrayView(TArrayView<FAnimDatabaseFrameAttribute> OutDirectionFrameAttributes, const UAnimDatabase* Database, const FAnimDatabaseFrameRanges& FrameRanges, const TArrayView<const FName> BoneNames, const FVector LocalDirection, const float RelativeTime)
+void UAnimDatabaseFrameAttributeLibrary::MakeBoneGlobalDirectionFrameAttributesFromNamesArrayView(const TArrayView<FAnimDatabaseFrameAttribute> OutDirectionFrameAttributes, const UAnimDatabase* Database, const FAnimDatabaseFrameRanges& FrameRanges, const TArrayView<const FName> BoneNames, const FVector LocalDirection, const float RelativeTime)
 {
 	if (!Database)
 	{
@@ -4029,9 +4076,9 @@ void UAnimDatabaseFrameAttributeLibrary::MakeFrameAttributesFromCurves(
 }
 
 void UAnimDatabaseFrameAttributeLibrary::MakeFrameAttributesFromCurvesArrayViews(
-	TArrayView<FAnimDatabaseFrameAttribute> OutFloatCurveValueFrameAttributes,
-	TArrayView<FAnimDatabaseFrameAttribute> OutFloatCurveVelocityFrameAttributes,
-	TArrayView<FAnimDatabaseFrameAttribute> OutBoolCurveActiveFrameAttributes,
+	const TArrayView<FAnimDatabaseFrameAttribute> OutFloatCurveValueFrameAttributes,
+	const TArrayView<FAnimDatabaseFrameAttribute> OutFloatCurveVelocityFrameAttributes,
+	const TArrayView<FAnimDatabaseFrameAttribute> OutBoolCurveActiveFrameAttributes,
 	const UAnimDatabase* Database,
 	const FAnimDatabaseFrameRanges& FrameRanges,
 	const TArrayView<const FName> CurveNames)
@@ -10153,6 +10200,11 @@ FAnimDatabaseFrameAttribute UAnimDatabaseFrameAttributeFunction::MakeFrameAttrib
 	return UAnimDatabaseFrameAttributeLibrary::MakeEmptyFrameAttribute();
 }
 
+FAnimDatabaseFrameAttribute UAnimDatabaseFrameAttributeFunction_Filter::MakeFrameAttribute_Implementation(const UAnimDatabase* Database, const FAnimDatabaseFrameRanges& FrameRanges) const
+{
+	return UAnimDatabaseFrameAttributeLibrary::MakeFrameAttributeFromFunction(Database, UAnimDatabaseFrameRangesLibrary::MakeFrameRangesFromFunction(Database, FrameRanges, FrameRangesFunction), FrameAttributeFunction);
+}
+
 FAnimDatabaseFrameRanges UAnimDatabaseFrameRangesFunction_Outliers::MakeFrameRanges_Implementation(const UAnimDatabase* Database, const FAnimDatabaseFrameRanges& FrameRanges) const
 {
 	return UAnimDatabaseFrameAttributeLibrary::FrameAttributeOutlierFrameRanges(
@@ -10212,6 +10264,13 @@ FAnimDatabaseFrameAttribute UAnimDatabaseFrameAttributeFunction_BoneGlobalTransf
 	return UAnimDatabaseFrameAttributeLibrary::FrameAttributeTransformMultiply(LocalOffset, UAnimDatabaseFrameAttributeLibrary::MakeBoneGlobalTransformFrameAttributeFromName(Database, FrameRanges, BoneName));
 }
 
+FAnimDatabaseFrameAttribute UAnimDatabaseFrameAttributeFunction_BoneRelativeTransform::MakeFrameAttribute_Implementation(const UAnimDatabase* Database, const FAnimDatabaseFrameRanges& FrameRanges) const
+{
+	TArray<FAnimDatabaseFrameAttribute, TInlineAllocator<2>> BoneTransforms; BoneTransforms.SetNum(2);
+	UAnimDatabaseFrameAttributeLibrary::MakeBoneGlobalTransformFrameAttributesFromNamesArrayView(BoneTransforms, Database, FrameRanges, { BoneName, RelativeBoneName });
+	return UAnimDatabaseFrameAttributeLibrary::FrameAttributeDivide(BoneTransforms[0], BoneTransforms[1]);
+}
+
 FAnimDatabaseFrameAttribute UAnimDatabaseFrameAttributeFunction_BoneGlobalLocation::MakeFrameAttribute_Implementation(const UAnimDatabase* Database, const FAnimDatabaseFrameRanges& FrameRanges) const
 {
 	return UAnimDatabaseFrameAttributeLibrary::FrameAttributeApplyTransformLocation(UAnimDatabaseFrameAttributeLibrary::MakeBoneGlobalTransformFrameAttributeFromName(Database, FrameRanges, BoneName), LocalOffset);
diff --git a/Engine/Plugins/Experimental/Animation/AnimDatabase/Source/AnimDatabase/Public/AnimDatabase.h b/Engine/Plugins/Experimental/Animation/AnimDatabase/Source/AnimDatabase/Public/AnimDatabase.h
--- a/Engine/Plugins/Experimental/Animation/AnimDatabase/Source/AnimDatabase/Public/AnimDatabase.h
+++ b/Engine/Plugins/Experimental/Animation/AnimDatabase/Source/AnimDatabase/Public/AnimDatabase.h
@@ -1007,6 +1007,15 @@ public:
 	static UE_API void DatabaseSetGlobalBoneTransforms(UAnimDatabase* Database, const TArray<FName>& BoneNames, const TArray<FAnimDatabaseFrameAttribute>& TransformFrameAttributes, const bool bShouldTransact = false);
 	static UE_API void DatabaseSetGlobalBoneTransformsFromArrayViews(UAnimDatabase* Database, const TArrayView<const FName> BoneNames, const TArrayView<const FAnimDatabaseFrameAttribute> TransformFrameAttributes, const bool bShouldTransact = false);
 
+	/** Set the given bone's local transforms in the database to the reference pose */
+	UFUNCTION(BlueprintCallable, Category = "AnimDatabase", meta = (DevelopmentOnly))
+	static UE_API void DatabaseResetBoneTransform(UAnimDatabase* Database, const FAnimDatabaseFrameRanges& FrameRanges, const FName BoneName, const bool bShouldTransact = false);
+
+	/** Set the given bones' local transforms in the database to the reference pose */
+	UFUNCTION(BlueprintCallable, Category = "AnimDatabase", meta = (DevelopmentOnly))
+	static UE_API void DatabaseResetBoneTransforms(UAnimDatabase* Database, const FAnimDatabaseFrameRanges& FrameRanges, const TArray<FName>& BoneNames, const bool bShouldTransact = false);
+	static UE_API void DatabaseResetBoneTransformsFromArrayViews(UAnimDatabase* Database, const FAnimDatabaseFrameRanges& FrameRanges, const TArrayView<const FName> BoneNames, const bool bShouldTransact = false);
+
 	/** Strips unused bone tracks from the data to reduce the file size on disk */
 	UFUNCTION(BlueprintCallable, Category = "AnimDatabase", meta = (DevelopmentOnly))
 	static UE_API void DatabaseStripUnusedBoneTracks(UAnimDatabase* Database, const FAnimDatabaseFrameRanges& FrameRanges, const bool bShouldTransact = false);
@@ -1793,6 +1802,68 @@ public:
 #endif
 };
 
+UENUM(BlueprintType)
+enum class EAnimDatabaseBoneSpace : uint8
+{
+	// Local (parent) space 
+	Local,
+	// Global (world) space
+	Global
+};
+
+/** Copy bone transforms from certain bones to others  */
+UCLASS(BlueprintType, Blueprintable, meta = (DisplayName = "Copy Bones"))
+class UAnimDatabaseFunction_CopyBones : public UAnimDatabaseFunction
+{
+	GENERATED_BODY()
+
+#if WITH_EDITORONLY_DATA
+
+public:
+
+	/** Bone Mapping */
+	UPROPERTY(EditAnywhere, Category = "Settings")
+	TMap<FName, FName> BoneMap;
+
+	/** Space to perform copying in */
+	UPROPERTY(EditAnywhere, Category = "Settings")
+	EAnimDatabaseBoneSpace Space = EAnimDatabaseBoneSpace::Global;
+
+#endif
+
+#if WITH_EDITOR
+
+public:
+
+	UE_API virtual void Run_Implementation(UAnimDatabase* InDatabase, const FAnimDatabaseFrameRanges& InFrameRanges, const bool bShouldTransact) override;
+
+#endif
+};
+
+/** Reset bone transforms to the reference pose  */
+UCLASS(BlueprintType, Blueprintable, meta = (DisplayName = "Reset Bones"))
+class UAnimDatabaseFunction_ResetBones : public UAnimDatabaseFunction
+{
+	GENERATED_BODY()
+
+#if WITH_EDITORONLY_DATA
+
+public:
+
+	UPROPERTY(EditAnywhere, Category = "Settings")
+	TArray<FName> Bones;
+
+#endif
+
+#if WITH_EDITOR
+
+public:
+
+	UE_API virtual void Run_Implementation(UAnimDatabase* InDatabase, const FAnimDatabaseFrameRanges& InFrameRanges, const bool bShouldTransact) override;
+
+#endif
+};
+
 /** Function to set bone to root location at frame  */
 UCLASS(BlueprintType, Blueprintable, meta = (DisplayName = "Set Bone to Root Transform at Frames"))
 class UAnimDatabaseFunction_SetBoneToRootTransformAtFrames : public UAnimDatabaseFunction
diff --git a/Engine/Plugins/Experimental/Animation/AnimDatabase/Source/AnimDatabase/Public/AnimDatabaseFrameAttribute.h b/Engine/Plugins/Experimental/Animation/AnimDatabase/Source/AnimDatabase/Public/AnimDatabaseFrameAttribute.h
--- a/Engine/Plugins/Experimental/Animation/AnimDatabase/Source/AnimDatabase/Public/AnimDatabaseFrameAttribute.h
+++ b/Engine/Plugins/Experimental/Animation/AnimDatabase/Source/AnimDatabase/Public/AnimDatabaseFrameAttribute.h
@@ -289,7 +289,7 @@ public:
 	/** Makes multiple float frame attributes from the given ranges and curve names. For when the curve is not present in the animations a default value of 0 will be used. */
 	UFUNCTION(BlueprintPure = true, Category = "AnimDatabase", meta = (NotBlueprintThreadSafe))
 	static UE_API void MakeFloatFrameAttributesFromCurves(TArray<FAnimDatabaseFrameAttribute>& OutFloatFrameAttributes, const UAnimDatabase* Database, const FAnimDatabaseFrameRanges& FrameRanges, const TArray<FName>& CurveNames);
-	static UE_API void MakeFloatFrameAttributesFromCurvesArrayView(TArrayView<FAnimDatabaseFrameAttribute> OutFloatFrameAttributes, const UAnimDatabase* Database, const FAnimDatabaseFrameRanges& FrameRanges, const TArrayView<const FName> CurveNames);
+	static UE_API void MakeFloatFrameAttributesFromCurvesArrayView(const TArrayView<FAnimDatabaseFrameAttribute> OutFloatFrameAttributes, const UAnimDatabase* Database, const FAnimDatabaseFrameRanges& FrameRanges, const TArrayView<const FName> CurveNames);
 
 	/** Makes a float frame attribute using the given curve name. For when the curve is not present in the database the attribute will be inactive. */
 	UFUNCTION(BlueprintPure = true, Category = "AnimDatabase", meta = (NotBlueprintThreadSafe))
@@ -298,7 +298,7 @@ public:
 	/** Makes multiple float frame attributes from the given ranges and curve names. For when the curve is not present in the animations the attribute will be inactive. */
 	UFUNCTION(BlueprintPure = true, Category = "AnimDatabase", meta = (NotBlueprintThreadSafe))
 	static UE_API void MakeFloatFrameAttributesFromCurvesWhenActive(TArray<FAnimDatabaseFrameAttribute>& OutFloatFrameAttributes, const UAnimDatabase* Database, const FAnimDatabaseFrameRanges& FrameRanges, const TArray<FName>& CurveNames);
-	static UE_API void MakeFloatFrameAttributesFromCurvesWhenActiveArrayView(TArrayView<FAnimDatabaseFrameAttribute> OutFloatFrameAttributes, const UAnimDatabase* Database, const FAnimDatabaseFrameRanges& FrameRanges, const TArrayView<const FName> CurveNames);
+	static UE_API void MakeFloatFrameAttributesFromCurvesWhenActiveArrayView(const TArrayView<FAnimDatabaseFrameAttribute> OutFloatFrameAttributes, const UAnimDatabase* Database, const FAnimDatabaseFrameRanges& FrameRanges, const TArrayView<const FName> CurveNames);
 
 	/** Makes a float frame attribute where the value of the attribute is given by the time since the start of the range for each frame in each range in the set. */
 	UFUNCTION(BlueprintPure = true, Category = "AnimDatabase")
@@ -596,13 +596,28 @@ public:
 		const TArray<int32>& BoneIndices,
 		const float RelativeTime = 0.0f);
 
-	static UE_API void MakeBoneLocalFrameAttributesFromArrayViews(
+	static UE_API void MakeBoneLocalTransformFrameAttributesFromArrayViews(
 		const TArrayView<FAnimDatabaseFrameAttribute> OutTransformFrameAttributes,
 		const UAnimDatabase* Database,
 		const FAnimDatabaseFrameRanges& FrameRanges,
 		const TArrayView<const int32> BoneIndices = {},
 		const float RelativeTime = 0.0f);
 
+	UFUNCTION(BlueprintPure = false, Category = "AnimDatabase", meta = (NotBlueprintThreadSafe, AutoCreateRefTerm = "BoneNames"))
+	static UE_API void MakeBoneLocalTransformFrameAttributesFromNames(
+		TArray<FAnimDatabaseFrameAttribute>& OutTransformFrameAttributes,
+		const UAnimDatabase* Database,
+		const FAnimDatabaseFrameRanges& FrameRanges,
+		const TArray<FName>& BoneNames,
+		const float RelativeTime = 0.0f);
+
+	static UE_API void MakeBoneLocalTransformFrameAttributesFromNamesArrayView(
+		const TArrayView<FAnimDatabaseFrameAttribute> OutTransformFrameAttributes,
+		const UAnimDatabase* Database,
+		const FAnimDatabaseFrameRanges& FrameRanges,
+		const TArrayView<const FName> BoneNames = {},
+		const float RelativeTime = 0.0f);
+
 	/** Makes a location frame attribute from the local transform of a bone */
 	UFUNCTION(BlueprintPure = false, Category = "AnimDatabase", meta = (NotBlueprintThreadSafe, AdvancedDisplay = 3))
 	static UE_API FAnimDatabaseFrameAttribute MakeBoneLocalTransformFrameAttribute(const UAnimDatabase* Database, const FAnimDatabaseFrameRanges& FrameRanges, const int32 BoneIndex, const float RelativeTime = 0.0f);
@@ -722,7 +737,7 @@ public:
 		const float RelativeTime = 0.0f);
 
 	static UE_API void MakeBoneGlobalTransformFrameAttributesFromNamesArrayView(
-		TArrayView<FAnimDatabaseFrameAttribute> OutTransformFrameAttributes,
+		const TArrayView<FAnimDatabaseFrameAttribute> OutTransformFrameAttributes,
 		const UAnimDatabase* Database,
 		const FAnimDatabaseFrameRanges& FrameRanges,
 		const TArrayView<const FName> BoneNames,
@@ -739,12 +754,12 @@ public:
 	/** Makes location frame attributes from the global location of bones */
 	UFUNCTION(BlueprintPure = true, Category = "AnimDatabase", meta = (NotBlueprintThreadSafe, AdvancedDisplay = 4))
 	static UE_API void MakeBoneGlobalLocationFrameAttributes(TArray<FAnimDatabaseFrameAttribute>& OutLocationFrameAttributes, const UAnimDatabase* Database, const FAnimDatabaseFrameRanges& FrameRanges, const TArray<int32>& BoneIndices, const float RelativeTime = 0.0f);
-	static UE_API void MakeBoneGlobalLocationFrameAttributesFromArrayViews(TArrayView<FAnimDatabaseFrameAttribute> OutLocationFrameAttributes, const UAnimDatabase* Database, const FAnimDatabaseFrameRanges& FrameRanges, const TArrayView<const int32> BoneIndices, const float RelativeTime = 0.0f);
+	static UE_API void MakeBoneGlobalLocationFrameAttributesFromArrayViews(const TArrayView<FAnimDatabaseFrameAttribute> OutLocationFrameAttributes, const UAnimDatabase* Database, const FAnimDatabaseFrameRanges& FrameRanges, const TArrayView<const int32> BoneIndices, const float RelativeTime = 0.0f);
 
 	/** Makes location frame attributes from the global location of bones, provided by name */
 	UFUNCTION(BlueprintPure = true, Category = "AnimDatabase", meta = (NotBlueprintThreadSafe, AdvancedDisplay = 4))
 	static UE_API void MakeBoneGlobalLocationFrameAttributesFromNames(TArray<FAnimDatabaseFrameAttribute>& OutLocationFrameAttributes, const UAnimDatabase* Database, const FAnimDatabaseFrameRanges& FrameRanges, const TArray<FName>& BoneNames, const float RelativeTime = 0.0f);
-	static UE_API void MakeBoneGlobalLocationFrameAttributesFromNamesArrayView(TArrayView<FAnimDatabaseFrameAttribute> OutLocationFrameAttributes, const UAnimDatabase* Database, const FAnimDatabaseFrameRanges& FrameRanges, const TArrayView<const FName> BoneNames, const float RelativeTime = 0.0f);
+	static UE_API void MakeBoneGlobalLocationFrameAttributesFromNamesArrayView(const TArrayView<FAnimDatabaseFrameAttribute> OutLocationFrameAttributes, const UAnimDatabase* Database, const FAnimDatabaseFrameRanges& FrameRanges, const TArrayView<const FName> BoneNames, const float RelativeTime = 0.0f);
 
 	/** Makes a direction frame attribute from the global direction of a bone */
 	UFUNCTION(BlueprintPure = true, Category = "AnimDatabase", meta = (NotBlueprintThreadSafe, AdvancedDisplay = 4))
@@ -757,12 +772,12 @@ public:
 	/** Makes direction frame attributes from the global direction of bones */
 	UFUNCTION(BlueprintPure = true, Category = "AnimDatabase", meta = (NotBlueprintThreadSafe, AdvancedDisplay = 5))
 	static UE_API void MakeBoneGlobalDirectionFrameAttributes(TArray<FAnimDatabaseFrameAttribute>& OutDirectionFrameAttributes, const UAnimDatabase* Database, const FAnimDatabaseFrameRanges& FrameRanges, const TArray<int32>& BoneIndices, const FVector LocalDirection = FVector::ForwardVector, const float RelativeTime = 0.0f);
-	static UE_API void MakeBoneGlobalDirectionFrameAttributesFromArrayViews(TArrayView<FAnimDatabaseFrameAttribute> OutDirectionFrameAttributes, const UAnimDatabase* Database, const FAnimDatabaseFrameRanges& FrameRanges, const TArrayView<const int32> BoneIndices, const FVector LocalDirection = FVector::ForwardVector, const float RelativeTime = 0.0f);
+	static UE_API void MakeBoneGlobalDirectionFrameAttributesFromArrayViews(const TArrayView<FAnimDatabaseFrameAttribute> OutDirectionFrameAttributes, const UAnimDatabase* Database, const FAnimDatabaseFrameRanges& FrameRanges, const TArrayView<const int32> BoneIndices, const FVector LocalDirection = FVector::ForwardVector, const float RelativeTime = 0.0f);
 
 	/** Makes direction frame attributes from the global direction of bones, provided by name */
 	UFUNCTION(BlueprintPure = true, Category = "AnimDatabase", meta = (NotBlueprintThreadSafe, AdvancedDisplay = 5))
 	static UE_API void MakeBoneGlobalDirectionFrameAttributesFromNames(TArray<FAnimDatabaseFrameAttribute>& OutDirectionFrameAttributes, const UAnimDatabase* Database, const FAnimDatabaseFrameRanges& FrameRanges, const TArray<FName>& BoneNames, const FVector LocalDirection = FVector::ForwardVector, const float RelativeTime = 0.0f);
-	static UE_API void MakeBoneGlobalDirectionFrameAttributesFromNamesArrayView(TArrayView<FAnimDatabaseFrameAttribute> OutDirectionFrameAttributes, const UAnimDatabase* Database, const FAnimDatabaseFrameRanges& FrameRanges, const TArrayView<const FName> BoneNames, const FVector LocalDirection = FVector::ForwardVector, const float RelativeTime = 0.0f);
+	static UE_API void MakeBoneGlobalDirectionFrameAttributesFromNamesArrayView(const TArrayView<FAnimDatabaseFrameAttribute> OutDirectionFrameAttributes, const UAnimDatabase* Database, const FAnimDatabaseFrameRanges& FrameRanges, const TArrayView<const FName> BoneNames, const FVector LocalDirection = FVector::ForwardVector, const float RelativeTime = 0.0f);
 
 	/** Makes a linear velocity frame attribute from the global linear velocity of a bone */
 	UFUNCTION(BlueprintPure = true, Category = "AnimDatabase", meta = (NotBlueprintThreadSafe, AdvancedDisplay = 3))
@@ -850,9 +865,9 @@ public:
 		const TArray<FName>& CurveNames);
 
 	static UE_API void MakeFrameAttributesFromCurvesArrayViews(
-		TArrayView<FAnimDatabaseFrameAttribute> OutFloatCurveValueFrameAttributes,
-		TArrayView<FAnimDatabaseFrameAttribute> OutFloatCurveVelocityFrameAttributes,
-		TArrayView<FAnimDatabaseFrameAttribute> OutBoolCurveActiveFrameAttributes,
+		const TArrayView<FAnimDatabaseFrameAttribute> OutFloatCurveValueFrameAttributes,
+		const TArrayView<FAnimDatabaseFrameAttribute> OutFloatCurveVelocityFrameAttributes,
+		const TArrayView<FAnimDatabaseFrameAttribute> OutBoolCurveActiveFrameAttributes,
 		const UAnimDatabase* Database,
 		const FAnimDatabaseFrameRanges& FrameRanges,
 		const TArrayView<const FName> CurveNames);
@@ -1959,6 +1974,27 @@ public:
 	UE_API FAnimDatabaseFrameAttribute MakeFrameAttribute(const UAnimDatabase* Database, const FAnimDatabaseFrameRanges& FrameRanges) const;
 };
 
+/** Computes an Attribute for the given Frame Ranges  */
+UCLASS(BlueprintType, Blueprintable, meta = (DisplayName = "Filter"))
+class UAnimDatabaseFrameAttributeFunction_Filter : public UAnimDatabaseFrameAttributeFunction
+{
+	GENERATED_BODY()
+
+public:
+
+	/** Frame Ranges to filter to */
+	UPROPERTY(EditAnywhere, Instanced, Category = "Settings")
+	TObjectPtr<UAnimDatabaseFrameRangesFunction> FrameRangesFunction;
+
+	/** Frame Attribute to create */
+	UPROPERTY(EditAnywhere, Instanced, Category = "Settings")
+	TObjectPtr<UAnimDatabaseFrameAttributeFunction> FrameAttributeFunction;
+
+public:
+
+	UE_API virtual FAnimDatabaseFrameAttribute MakeFrameAttribute_Implementation(const UAnimDatabase* Database, const FAnimDatabaseFrameRanges& FrameRanges) const override;
+};
+
 /** Finds all the potential outlier frame ranges in a database */
 UCLASS(BlueprintType, Blueprintable, meta = (DisplayName = "Outliers"))
 class UAnimDatabaseFrameRangesFunction_Outliers : public UAnimDatabaseFrameRangesFunction
@@ -2111,6 +2147,27 @@ public:
 	UE_API virtual FAnimDatabaseFrameAttribute MakeFrameAttribute_Implementation(const UAnimDatabase* Database, const FAnimDatabaseFrameRanges& FrameRanges) const override;
 };
 
+/** Computes a bone global transform attribute relative to another  */
+UCLASS(BlueprintType, Blueprintable, meta = (DisplayName = "Bone Relative Transform"))
+class UAnimDatabaseFrameAttributeFunction_BoneRelativeTransform : public UAnimDatabaseFrameAttributeFunction
+{
+	GENERATED_BODY()
+
+public:
+
+	/** Name of the bone to compute the transform of */
+	UPROPERTY(EditAnywhere, Category = "Settings")
+	FName BoneName = NAME_None;
+
+	/** Name of the bone to compute the transform relative to */
+	UPROPERTY(EditAnywhere, Category = "Settings")
+	FName RelativeBoneName = NAME_None;
+
+public:
+
+	UE_API virtual FAnimDatabaseFrameAttribute MakeFrameAttribute_Implementation(const UAnimDatabase* Database, const FAnimDatabaseFrameRanges& FrameRanges) const override;
+};
+
 /** Computes a bone global location attribute */
 UCLASS(BlueprintType, Blueprintable, meta = (DisplayName = "Bone Global Location"))
 class UAnimDatabaseFrameAttributeFunction_BoneGlobalLocation : public UAnimDatabaseFrameAttributeFunction
diff --git a/Engine/Plugins/Experimental/Animation/AnimGen/Source/AnimGen/Private/AnimNode_AnimGenController.cpp b/Engine/Plugins/Experimental/Animation/AnimGen/Source/AnimGen/Private/AnimNode_AnimGenController.cpp
--- a/Engine/Plugins/Experimental/Animation/AnimGen/Source/AnimGen/Private/AnimNode_AnimGenController.cpp
+++ b/Engine/Plugins/Experimental/Animation/AnimGen/Source/AnimGen/Private/AnimNode_AnimGenController.cpp
@@ -670,6 +670,49 @@ void FAnimNode_AnimGenController::Evaluate_AnyThread(FPoseContext& Output)
 		}
 	}
 
+	// Output Additional bone transforms
+
+	for (const TPair<FName, FName>& AttributeTransform : BoneTransformMap)
+	{
+		const int32 AttributeIdx = ActiveController->AutoEncoder->FindAttributeIndex(AttributeTransform.Key);
+
+		if (AttributeIdx == INDEX_NONE)
+		{
+			Output.LogMessage(EMessageSeverity::Warning, FText::Format(
+				LOCTEXT("AttributeTransformMapAttributeNotFoundWarning", "Cannot create transform from attribute '{0}' - not found in AutoEncoder."),
+				FText::FromString(AttributeTransform.Key.ToString())));
+			continue;
+		}
+
+		if (ActiveController->AutoEncoder->GetAttributeType(AttributeIdx) != EAnimDatabaseAttributeType::Transform)
+		{
+			Output.LogMessage(EMessageSeverity::Warning, FText::Format(
+				LOCTEXT("AttributeTransformMapAttributeWrongTypeWarning", "Cannot create transform from attribute '{0}' - attribute is not Transform type."),
+				FText::FromString(AttributeTransform.Key.ToString())));
+			continue;
+		}
+
+		if (PoseState.CurrPoseData->AttributeData.GetAttributeActive(0, AttributeIdx))
+		{
+			const int32 PoseBoneIndex = BoneContainer.GetPoseBoneIndexForBoneName(AttributeTransform.Value);
+			
+			if (PoseBoneIndex != INDEX_NONE)
+			{
+				const FCompactPoseBoneIndex CompactBoneIndex = BoneContainer.GetCompactPoseIndexFromSkeletonPoseIndex(FSkeletonPoseBoneIndex(PoseBoneIndex));
+
+				if (CompactBoneIndex.IsValid() && CompactBoneIndex >= 0 && CompactBoneIndex < CompactBones.Num())
+				{
+					const FTransform3f BoneTransform = PoseState.CurrPoseData->AttributeData.GetTransform(0, AttributeIdx);
+
+					CompactBones[CompactBoneIndex.GetInt()] = FTransform(
+						((FQuat)BoneTransform.GetRotation()).GetNormalized(),
+						(FVector)BoneTransform.GetLocation(),
+						(FVector)BoneTransform.GetScale3D());
+				}
+			}
+		}
+	}
+
 	// Output Curves
 
 	OutputCurves.Reset();
diff --git a/Engine/Plugins/Experimental/Animation/AnimGen/Source/AnimGen/Public/AnimNode_AnimGenController.h b/Engine/Plugins/Experimental/Animation/AnimGen/Source/AnimGen/Public/AnimNode_AnimGenController.h
--- a/Engine/Plugins/Experimental/Animation/AnimGen/Source/AnimGen/Public/AnimNode_AnimGenController.h
+++ b/Engine/Plugins/Experimental/Animation/AnimGen/Source/AnimGen/Public/AnimNode_AnimGenController.h
@@ -81,6 +81,10 @@ public:
 	UPROPERTY(EditAnywhere, Category = Curves)
 	TMap<FName, FName> CurveMap;
 
+	/** Map of attribute names to bone transforms to produce */
+	UPROPERTY(EditAnywhere, Category = Curves)
+	TMap<FName, FName> BoneTransformMap;
+
 private:
 
 	/** Internal buffer of AnimNotifyEvents */
