AnimDatabase: Fixed minor threading issue in Frame Attribute code

#Synced-CL 51845390

diff --git a/Engine/Plugins/Experimental/Animation/AnimDatabase/Source/AnimDatabase/Private/AnimDatabase.cpp b/Engine/Plugins/Experimental/Animation/AnimDatabase/Source/AnimDatabase/Private/AnimDatabase.cpp
--- a/Engine/Plugins/Experimental/Animation/AnimDatabase/Source/AnimDatabase/Private/AnimDatabase.cpp
+++ b/Engine/Plugins/Experimental/Animation/AnimDatabase/Source/AnimDatabase/Private/AnimDatabase.cpp
@@ -2847,7 +2847,7 @@ void UAnimDatabase::WaitForCompressionOnAnimSequence(const int32 SequenceIdx) co
 	UAnimSequence* Sequence = GetAnimSequence(SequenceIdx);
 
 #if WITH_EDITOR
-	if (Sequence && Sequence->IsCompressedDataOutOfDate())
+	if (IsInGameThread() && Sequence && Sequence->IsCompressedDataOutOfDate())
 	{
 		Sequence->WaitOnExistingCompression(true);
 	}
@@ -7105,6 +7105,11 @@ void UAnimDatabaseFunction_AddCurveFromFrameAttribute::Run_Implementation(UAnimD
 
 void  UAnimDatabaseFunction_AddContactCurves::Run_Implementation(UAnimDatabase* InDatabase, const FAnimDatabaseFrameRanges& InFrameRanges, const bool bShouldTransact)
 {
+	if (!InDatabase) { return; }
+	if (!InFrameRanges.IsValid()) { return; }
+
+	InDatabase->WaitForCompressionOnAnimSequencesFromArrayView(InFrameRanges.FrameRangeSet->GetEntrySequences());
+
 	UE::Tasks::TTask<FAnimDatabaseFrameAttribute> LeftContactCurve = UE::Tasks::Launch(UE_SOURCE_LOCATION, [this, InDatabase, InFrameRanges]() {
 		
 		return UAnimDatabaseFrameAttributeLibrary::FrameAttributeContactCurve(
diff --git a/Engine/Plugins/Experimental/Animation/AnimDatabase/Source/AnimDatabase/Private/AnimDatabaseFrameAttribute.cpp b/Engine/Plugins/Experimental/Animation/AnimDatabase/Source/AnimDatabase/Private/AnimDatabaseFrameAttribute.cpp
--- a/Engine/Plugins/Experimental/Animation/AnimDatabase/Source/AnimDatabase/Private/AnimDatabaseFrameAttribute.cpp
+++ b/Engine/Plugins/Experimental/Animation/AnimDatabase/Source/AnimDatabase/Private/AnimDatabaseFrameAttribute.cpp
@@ -9966,6 +9966,20 @@ FAnimDatabaseFrameAttribute UAnimDatabaseFrameAttributeLibrary::FrameAttributeSa
 	const int32 DirectionSavGolPolynomialDegree,
 	const bool bGaussianWindowed)
 {
+	if (!Database)
+	{
+		UE_LOGFMT(LogAnimDatabase, Error, "FrameAttributeSavGolSmoothedRootTransform: Database is nullptr.");
+		return FAnimDatabaseFrameAttribute();
+	}
+
+	if (!FrameRanges.IsValid())
+	{
+		UE_LOGFMT(LogAnimDatabase, Error, "FrameAttributeSavGolSmoothedRootTransform: Invalid FrameRanges.");
+		return FAnimDatabaseFrameAttribute();
+	}
+
+	Database->WaitForCompressionOnAnimSequencesFromArrayView(FrameRanges.FrameRangeSet->GetEntrySequences());
+
 	UE::Tasks::TTask<FAnimDatabaseFrameAttribute> RootLocationTask = UE::Tasks::Launch(UE_SOURCE_LOCATION, [Database, &FrameRanges, RootLocationProjectionBoneName, LocationSavGolFrameNum, LocationSavGolPolynomialDegree, bGaussianWindowed]()
 		{
 			const FAnimDatabaseFrameAttribute RootLocation = MakeBoneGlobalLocationFrameAttributeFromName(Database, FrameRanges, RootLocationProjectionBoneName);
@@ -9992,6 +10006,20 @@ FAnimDatabaseFrameAttribute UAnimDatabaseFrameAttributeLibrary::FrameAttributeGa
 	const float LocationSmoothingAmount,
 	const float RotationSmoothingAmount)
 {
+	if (!Database)
+	{
+		UE_LOGFMT(LogAnimDatabase, Error, "FrameAttributeGaussianSmoothedRootTransform: Database is nullptr.");
+		return FAnimDatabaseFrameAttribute();
+	}
+
+	if (!FrameRanges.IsValid())
+	{
+		UE_LOGFMT(LogAnimDatabase, Error, "FrameAttributeGaussianSmoothedRootTransform: Invalid FrameRanges.");
+		return FAnimDatabaseFrameAttribute();
+	}
+
+	Database->WaitForCompressionOnAnimSequencesFromArrayView(FrameRanges.FrameRangeSet->GetEntrySequences());
+
 	UE::Tasks::TTask<FAnimDatabaseFrameAttribute> RootLocationTask = UE::Tasks::Launch(UE_SOURCE_LOCATION, [Database, &FrameRanges, RootLocationProjectionBoneName, LocationSmoothingAmount]()
 		{
 			const FAnimDatabaseFrameAttribute RootLocation = MakeBoneGlobalLocationFrameAttributeFromName(Database, FrameRanges, RootLocationProjectionBoneName);
