AnimDatabase: Various improvement to Database functions to better support undo and other functionality

#Synced-CL 51945476

diff --git a/Engine/Plugins/Experimental/Animation/AnimDatabase/Source/AnimDatabase/AnimDatabase.Build.cs b/Engine/Plugins/Experimental/Animation/AnimDatabase/Source/AnimDatabase/AnimDatabase.Build.cs
--- a/Engine/Plugins/Experimental/Animation/AnimDatabase/Source/AnimDatabase/AnimDatabase.Build.cs
+++ b/Engine/Plugins/Experimental/Animation/AnimDatabase/Source/AnimDatabase/AnimDatabase.Build.cs
@@ -41,12 +41,16 @@ public class AnimDatabase : ModuleRules
 				"Json",
 				"JsonUtilities",
 				"DrawDebugLibrary",
-				"RigVM",
+				"RigVM"
 				// ... add private dependencies that you statically link with here ...	
 			}
 			);
-		
-		
+
+		if (Target.bBuildEditor)
+		{
+			PrivateDependencyModuleNames.Add("UnrealEd");
+		}
+
 		DynamicallyLoadedModuleNames.AddRange(
 			new string[]
 			{
diff --git a/Engine/Plugins/Experimental/Animation/AnimDatabase/Source/AnimDatabase/Private/AnimDatabase.cpp b/Engine/Plugins/Experimental/Animation/AnimDatabase/Source/AnimDatabase/Private/AnimDatabase.cpp
--- a/Engine/Plugins/Experimental/Animation/AnimDatabase/Source/AnimDatabase/Private/AnimDatabase.cpp
+++ b/Engine/Plugins/Experimental/Animation/AnimDatabase/Source/AnimDatabase/Private/AnimDatabase.cpp
@@ -23,6 +23,9 @@
 #include "UObject/LinkerLoad.h"
 #include "UObject/Package.h"
 #include "Async/ParallelFor.h"
+#if WITH_EDITOR
+#include "ScopedTransaction.h"
+#endif
 
 #include "Dom/JsonObject.h"
 #include "Misc/Paths.h"
@@ -5859,17 +5862,17 @@ void UAnimDatabaseLibrary::DatabaseExportRootMotion(UAnimDatabase* Database, con
 #endif
 }
 
-void UAnimDatabaseLibrary::DatabaseAddCurve(UAnimDatabase* Database, const FName CurveName, const FAnimDatabaseFrameAttribute& FloatFrameAttribute, const bool bSparseKeys, const bool bShouldTransact)
+void UAnimDatabaseLibrary::DatabaseAddCurve(UAnimDatabase* Database, const FName CurveName, const FAnimDatabaseFrameAttribute& FloatFrameAttribute, const FLinearColor Color, const bool bSparseKeys, const bool bShouldTransact)
 {
-	DatabaseAddCurvesFromArrayViews(Database, { CurveName }, { FloatFrameAttribute }, bSparseKeys, bShouldTransact);
+	DatabaseAddCurvesFromArrayViews(Database, { CurveName }, { FloatFrameAttribute }, { Color }, bSparseKeys, bShouldTransact);
 }
 
-void UAnimDatabaseLibrary::DatabaseAddCurves(UAnimDatabase* Database, const TArray<FName>& CurveNames, const TArray<FAnimDatabaseFrameAttribute>& FloatFrameAttributes, const bool bSparseKeys, const bool bShouldTransact)
+void UAnimDatabaseLibrary::DatabaseAddCurves(UAnimDatabase* Database, const TArray<FName>& CurveNames, const TArray<FAnimDatabaseFrameAttribute>& FloatFrameAttributes, const TArray<FLinearColor>& Colors, const bool bSparseKeys, const bool bShouldTransact)
 {
-	DatabaseAddCurvesFromArrayViews(Database, CurveNames, FloatFrameAttributes, bSparseKeys, bShouldTransact);
+	DatabaseAddCurvesFromArrayViews(Database, CurveNames, FloatFrameAttributes, Colors, bSparseKeys, bShouldTransact);
 }
 
-void UAnimDatabaseLibrary::DatabaseAddCurvesFromArrayViews(UAnimDatabase* Database, const TArrayView<const FName> CurveNames, const TArrayView<const FAnimDatabaseFrameAttribute> FloatFrameAttributes, const bool bSparseKeys, const bool bShouldTransact)
+void UAnimDatabaseLibrary::DatabaseAddCurvesFromArrayViews(UAnimDatabase* Database, const TArrayView<const FName> CurveNames, const TArrayView<const FAnimDatabaseFrameAttribute> FloatFrameAttributes, const TArrayView<const FLinearColor> Colors, const bool bSparseKeys, const bool bShouldTransact)
 {
 #if WITH_EDITOR
 	TRACE_CPUPROFILER_EVENT_SCOPE(UAnimDatabaseLibrary::DatabaseAddCurvesFromArrayViews);
@@ -5935,6 +5938,11 @@ void UAnimDatabaseLibrary::DatabaseAddCurvesFromArrayViews(UAnimDatabase* Databa
 				Controller->AddCurve(CurveId, AACF_DefaultCurve, bShouldTransact);
 			}
 
+			if (Colors.IsValidIndex(CurveIdx))
+			{
+				Controller->SetCurveColor(CurveId, Colors[CurveIdx]);
+			}
+
 			// Add Curve Keys
 			const int32 RangeFrameNum = FilteredFrameAttribute.FrameAttribute->FrameRangeSet.GetAllRangeLengths()[RangeIdx];
 			const int32 FrameAttributeRangeOffset = FilteredFrameAttribute.FrameAttribute->FrameRangeSet.GetAllRangeOffsets()[RangeIdx];
@@ -5943,16 +5951,18 @@ void UAnimDatabaseLibrary::DatabaseAddCurvesFromArrayViews(UAnimDatabase* Databa
 
 			const int32 StartFrame = FilteredFrameAttribute.FrameAttribute->FrameRangeSet.GetAllRangeStarts()[RangeIdx];
 
-			float PrevValue = UE_MAX_FLT;
+			// TODO: Check for Linear Interpolation
+			float PrevSetValue = UE_MAX_FLT;
 			for (int32 FrameIdx = 0; FrameIdx < RangeFrameNum; FrameIdx++)
 			{
 				const float Time = (StartFrame + FrameIdx) / FMath::Max(Database->GetFrameRate().AsDecimal(), UE_SMALL_NUMBER);
-				const float Value = FilteredFrameAttribute.FrameAttribute->AttributeData[0][FrameAttributeRangeOffset + FrameIdx];
-				if (!bSparseKeys || PrevValue != Value)
+				const float CurrValue = FilteredFrameAttribute.FrameAttribute->AttributeData[0][FrameAttributeRangeOffset + FrameIdx];
+				const float NextValue = FilteredFrameAttribute.FrameAttribute->AttributeData[0][FrameAttributeRangeOffset + FMath::Min(FrameIdx + 1, FMath::Max(RangeFrameNum - 1, 0))];
+				if (!bSparseKeys || (CurrValue != PrevSetValue) || (NextValue != PrevSetValue))
 				{
-					RichCurveKeys.Add(FRichCurveKey(Time, Value));
+					RichCurveKeys.Add(FRichCurveKey(Time, CurrValue));
+					PrevSetValue = CurrValue;
 				}
-				PrevValue = Value;
 			}
 
 			Controller->SetCurveKeys(CurveId, RichCurveKeys, bShouldTransact);
@@ -5963,17 +5973,17 @@ void UAnimDatabaseLibrary::DatabaseAddCurvesFromArrayViews(UAnimDatabase* Databa
 #endif
 }
 
-void UAnimDatabaseLibrary::DatabaseAddAnimNotifyState(UAnimDatabase* Database, const FName AnimNotifyStateTrackName, const TSubclassOf<UAnimNotifyState> Class, const FAnimDatabaseFrameRanges& FrameRanges)
+void UAnimDatabaseLibrary::DatabaseAddAnimNotifyState(UAnimDatabase* Database, const FName AnimNotifyStateTrackName, const TSubclassOf<UAnimNotifyState> Class, const FAnimDatabaseFrameRanges& FrameRanges, const FLinearColor Color, const bool bShouldTransact)
 {
-	DatabaseAddAnimNotifyStatesFromArrayViews(Database, { AnimNotifyStateTrackName }, { Class }, { FrameRanges });
+	DatabaseAddAnimNotifyStatesFromArrayViews(Database, { AnimNotifyStateTrackName }, { Class }, { FrameRanges }, { Color }, bShouldTransact);
 }
 
-void UAnimDatabaseLibrary::DatabaseAddAnimNotifyStates(UAnimDatabase* Database, const TArray<FName>& AnimNotifyStateTrackNames, const TArray<TSubclassOf<UAnimNotifyState>>& Classes, const TArray<FAnimDatabaseFrameRanges>& FrameRanges)
+void UAnimDatabaseLibrary::DatabaseAddAnimNotifyStates(UAnimDatabase* Database, const TArray<FName>& AnimNotifyStateTrackNames, const TArray<TSubclassOf<UAnimNotifyState>>& Classes, const TArray<FAnimDatabaseFrameRanges>& FrameRanges, const TArray<FLinearColor>& Colors, const bool bShouldTransact)
 {
-	DatabaseAddAnimNotifyStatesFromArrayViews(Database, AnimNotifyStateTrackNames, Classes, FrameRanges);
+	DatabaseAddAnimNotifyStatesFromArrayViews(Database, AnimNotifyStateTrackNames, Classes, FrameRanges, Colors, bShouldTransact);
 }
 
-void UAnimDatabaseLibrary::DatabaseAddAnimNotifyStatesFromArrayViews(UAnimDatabase* Database, const TArrayView<const FName> AnimNotifyStateTrackNames, const TArrayView<const TSubclassOf<UAnimNotifyState>> Classes, const TArrayView<const FAnimDatabaseFrameRanges> FrameRanges)
+void UAnimDatabaseLibrary::DatabaseAddAnimNotifyStatesFromArrayViews(UAnimDatabase* Database, const TArrayView<const FName> AnimNotifyStateTrackNames, const TArrayView<const TSubclassOf<UAnimNotifyState>> Classes, const TArrayView<const FAnimDatabaseFrameRanges> FrameRanges, const TArrayView<const FLinearColor> Colors, const bool bShouldTransact)
 {
 #if WITH_EDITOR
 	TRACE_CPUPROFILER_EVENT_SCOPE(UAnimDatabaseLibrary::DatabaseAddAnimNotifyStatesFromArrayViews);
@@ -5996,19 +6006,18 @@ void UAnimDatabaseLibrary::DatabaseAddAnimNotifyStatesFromArrayViews(UAnimDataba
 	for (int32 TrackIdx = 0; TrackIdx < TrackNum; TrackIdx++)
 	{
 		const FAnimDatabaseFrameRanges FilteredFrameRanges = UAnimDatabaseFrameRangesLibrary::MakeFrameRangesFromNotMirrored(Database, FrameRanges[TrackIdx]);
-		const int32 TotalRangeNum = FilteredFrameRanges.FrameRangeSet->GetTotalRangeNum();
 
-		TLearningArray<1, int32> RangeSequenceIndices;
-		RangeSequenceIndices.SetNumUninitialized({ TotalRangeNum });
-		UE::Learning::FrameRangeSet::AllRangeSequences(RangeSequenceIndices, *FilteredFrameRanges.FrameRangeSet);
+		const int32 EntryNum = FilteredFrameRanges.FrameRangeSet->GetEntryNum();
 
-		for (int32 RangeIdx = 0; RangeIdx < TotalRangeNum; RangeIdx++)
+		for (int32 EntryIdx = 0; EntryIdx < EntryNum; EntryIdx++)
 		{
-			TRACE_CPUPROFILER_EVENT_SCOPE(UAnimDatabaseLibrary::DatabaseAddAnimNotifyStatesFromArrayViews::Range);
-			check(!Database->GetIsMirrored(RangeSequenceIndices[RangeIdx]));
+			check(!Database->GetIsMirrored(FilteredFrameRanges.FrameRangeSet->GetEntrySequence(EntryIdx)));
 
-			if (UAnimSequence* AnimSequence = Database->GetAnimSequence(RangeSequenceIndices[RangeIdx]))
+			if (UAnimSequence* AnimSequence = Database->GetAnimSequence(FilteredFrameRanges.FrameRangeSet->GetEntrySequence(EntryIdx)))
 			{
+				FScopedTransaction Transaction(LOCTEXT("AddAnimNotifyStatesTransaction", "Adding Anim Notify States"), bShouldTransact);
+				AnimSequence->Modify();
+
 				int32 AnimNotifyTrack = AnimSequence->AnimNotifyTracks.IndexOfByPredicate(
 					[&](const FAnimNotifyTrack& Track)
 					{
@@ -6020,188 +6029,264 @@ void UAnimDatabaseLibrary::DatabaseAddAnimNotifyStatesFromArrayViews(UAnimDataba
 					AnimNotifyTrack = AnimSequence->AnimNotifyTracks.Add(FAnimNotifyTrack(AnimNotifyStateTrackNames[TrackIdx], FLinearColor::Red));
 				}
 
-				const float StartTime = FilteredFrameRanges.FrameRangeSet->GetAllRangeStarts()[RangeIdx] / FMath::Max(Database->GetFrameRate().AsDecimal(), UE_SMALL_NUMBER);
-				const float Duration = (FilteredFrameRanges.FrameRangeSet->GetAllRangeLengths()[RangeIdx] - 1) / FMath::Max(Database->GetFrameRate().AsDecimal(), UE_SMALL_NUMBER);
-
-				FAnimNotifyEvent& NewEvent = AnimSequence->Notifies.AddDefaulted_GetRef();
+				if (Colors.IsValidIndex(TrackIdx))
+				{
+					AnimSequence->AnimNotifyTracks[AnimNotifyTrack].TrackColor = Colors[TrackIdx];
+				}
 
-				NewEvent.NotifyName = NAME_None;
-				NewEvent.Link(AnimSequence, StartTime);
-				NewEvent.TriggerTimeOffset = GetTriggerTimeOffsetForType(AnimSequence->CalculateOffsetForNotify(StartTime));
-				NewEvent.TrackIndex = AnimNotifyTrack;
-				NewEvent.Notify = nullptr;
-				NewEvent.Guid = FGuid::NewGuid();
+				const int32 EntryRangeNum = FilteredFrameRanges.FrameRangeSet->GetEntryRangeNum(EntryIdx);
 
-				if (Classes[TrackIdx])
+				for (int32 RangeIdx = 0; RangeIdx < EntryRangeNum; RangeIdx++)
 				{
-					NewEvent.NotifyStateClass = NewObject<UAnimNotifyState>(AnimSequence, Classes[TrackIdx], NAME_None, RF_Transactional);
-					if (NewEvent.NotifyStateClass)
+					const float StartTime = FilteredFrameRanges.FrameRangeSet->GetEntryRangeStart(EntryIdx, RangeIdx) / FMath::Max(Database->GetFrameRate().AsDecimal(), UE_SMALL_NUMBER);
+					const float Duration = (FilteredFrameRanges.FrameRangeSet->GetEntryRangeLength(EntryIdx, RangeIdx) - 1) / FMath::Max(Database->GetFrameRate().AsDecimal(), UE_SMALL_NUMBER);
+
+					FAnimNotifyEvent& NewEvent = AnimSequence->Notifies.AddDefaulted_GetRef();
+
+					NewEvent.NotifyName = NAME_None;
+					NewEvent.Link(AnimSequence, StartTime);
+					NewEvent.TriggerTimeOffset = GetTriggerTimeOffsetForType(AnimSequence->CalculateOffsetForNotify(StartTime));
+					NewEvent.TrackIndex = AnimNotifyTrack;
+					NewEvent.Notify = nullptr;
+					NewEvent.Guid = FGuid::NewGuid();
+
+					if (Classes[TrackIdx])
 					{
-						NewEvent.NotifyName = FName(*NewEvent.NotifyStateClass->GetNotifyName());
-						NewEvent.SetDuration(Duration);
-						NewEvent.EndLink.Link(AnimSequence, NewEvent.EndLink.GetTime());
+						NewEvent.NotifyStateClass = NewObject<UAnimNotifyState>(AnimSequence, Classes[TrackIdx], NAME_None, RF_Transactional);
+						if (NewEvent.NotifyStateClass)
+						{
+							NewEvent.NotifyName = FName(*NewEvent.NotifyStateClass->GetNotifyName());
+							NewEvent.SetDuration(Duration);
+							NewEvent.EndLink.Link(AnimSequence, NewEvent.EndLink.GetTime());
+						}
+					}
+					else
+					{
+						NewEvent.NotifyStateClass = nullptr;
 					}
-				}
-				else
-				{
-					NewEvent.NotifyStateClass = nullptr;
 				}
 
-				// Refresh all cached data
 				AnimSequence->RefreshCacheData();
-				AnimSequence->Modify(true);
+				AnimSequence->MarkPackageDirty();
 			}
 		}
 	}
 #endif
 }
 
-void UAnimDatabaseLibrary::DatabaseAddAnimNotifyObject(UAnimDatabase* Database, const FAnimDatabaseFrames& Frames, const FName& TrackName, const UAnimNotify* AnimNotify)
+void UAnimDatabaseLibrary::DatabaseAddAnimNotifyObject(UAnimDatabase* Database, const FName& TrackName, UAnimNotify* AnimNotify, const FAnimDatabaseFrames& Frames, const FLinearColor Color, const bool bShouldTransact)
+{
+	DatabaseAddAnimNotifyObjectsFromArrayViews(Database, { TrackName }, { AnimNotify }, { Frames }, { Color }, bShouldTransact);
+}
+
+void UAnimDatabaseLibrary::DatabaseAddAnimNotifyObjects(UAnimDatabase* Database, const TArray<FName>& AnimNotifyTrackNames, const TArray<UAnimNotify*>& AnimNotifies, const TArray<FAnimDatabaseFrames>& Frames, const TArray<FLinearColor>& Colors, const bool bShouldTransact)
+{
+	DatabaseAddAnimNotifyObjectsFromArrayViews(Database, AnimNotifyTrackNames, AnimNotifies, Frames, Colors, bShouldTransact);
+}
+
+void UAnimDatabaseLibrary::DatabaseAddAnimNotifyObjectsFromArrayViews(UAnimDatabase* Database, const TArrayView<const FName> AnimNotifyTrackNames, const TArrayView<UAnimNotify* const> AnimNotifies, const TArrayView<const FAnimDatabaseFrames> Frames, const TArrayView<const FLinearColor> Colors, const bool bShouldTransact)
 {
 #if WITH_EDITOR
 	TRACE_CPUPROFILER_EVENT_SCOPE(UAnimDatabaseLibrary::DatabaseAddAnimNotifyObject);
 
 	if (!Database)
 	{
-		UE_LOGFMT(LogAnimDatabase, Error, "DatabaseAddAnimNotifyObject: Database is nullptr.");
+		UE_LOGFMT(LogAnimDatabase, Error, "DatabaseAddAnimNotifyObjectsFromArrayViews: Database is nullptr.");
 		return;
 	}
 
-	if (!AnimNotify)
+	if (AnimNotifyTrackNames.Num() != AnimNotifies.Num() ||
+		AnimNotifyTrackNames.Num() != Frames.Num())
 	{
-		UE_LOGFMT(LogAnimDatabase, Error, "DatabaseAddAnimNotifyObject: AnimNotify is nullptr.");
+		UE_LOGFMT(LogAnimDatabase, Error, "DatabaseAddAnimNotifyObjectsFromArrayViews: Number of Track Names, Notifies, and Frames must match. Got {Tracks}, {Notifies}, and {Frames}.", AnimNotifyTrackNames.Num(), AnimNotifies.Num(), Frames.Num());
 		return;
 	}
 
-	const FAnimDatabaseFrames FilteredFrames = UAnimDatabaseFrameRangesLibrary::MakeFramesAtFrameRangesStarts(
-		UAnimDatabaseFrameRangesLibrary::MakeFrameRangesFromNotMirrored(Database, UAnimDatabaseFrameRangesLibrary::MakeFrameRangesFromFrames(Frames)));
-	const int32 TotalFrameNum = FilteredFrames.FrameSet->GetTotalFrameNum();
-
-	TLearningArray<1, int32> FrameSequenceIndices;
-	FrameSequenceIndices.SetNumUninitialized({ TotalFrameNum });
-	UE::Learning::FrameSet::AllFrameSequences(FrameSequenceIndices, *FilteredFrames.FrameSet);
+	const int32 TrackNum = AnimNotifyTrackNames.Num();
 
-	for (int32 FrameIdx = 0; FrameIdx < TotalFrameNum; FrameIdx++)
+	for (int32 TrackIdx = 0; TrackIdx < TrackNum; TrackIdx++)
 	{
-		TRACE_CPUPROFILER_EVENT_SCOPE(UAnimDatabaseLibrary::DatabaseAddAnimNotifyObject::Frame);
-		check(!Database->GetIsMirrored(FrameSequenceIndices[FrameIdx]));
+		if (!AnimNotifies[TrackIdx])
+		{
+			UE_LOGFMT(LogAnimDatabase, Error, "DatabaseAddAnimNotifyObjectsFromArrayViews: AnimNotify is nullptr.");
+			continue;
+		}
 
-		if (UAnimSequence* AnimSequence = Database->GetAnimSequence(FrameSequenceIndices[FrameIdx]))
+		const FAnimDatabaseFrames FilteredFrames = UAnimDatabaseFrameRangesLibrary::MakeFramesAtFrameRangesStarts(
+			UAnimDatabaseFrameRangesLibrary::MakeFrameRangesFromNotMirrored(Database, UAnimDatabaseFrameRangesLibrary::MakeFrameRangesFromFrames(Frames[TrackIdx])));
+		
+		const int32 EntryNum = FilteredFrames.FrameSet->GetEntryNum();
+
+		for (int32 EntryIdx = 0; EntryIdx < EntryNum; EntryIdx++)
 		{
-			int32 AnimNotifyTrack = AnimSequence->AnimNotifyTracks.IndexOfByPredicate(
-				[&](const FAnimNotifyTrack& Track)
-				{
-					return Track.TrackName == TrackName;
-				});
+			check(!Database->GetIsMirrored(FilteredFrames.FrameSet->GetEntrySequence(EntryIdx)));
 
-			if (AnimNotifyTrack == INDEX_NONE)
+			if (UAnimSequence* AnimSequence = Database->GetAnimSequence(FilteredFrames.FrameSet->GetEntrySequence(EntryIdx)))
 			{
-				AnimNotifyTrack = AnimSequence->AnimNotifyTracks.Add(FAnimNotifyTrack(TrackName, FLinearColor::Red));
-			}
+				FScopedTransaction Transaction(LOCTEXT("AddAnimNotifyObjectsTransaction", "Adding Anim Notify Objects"), bShouldTransact);
+				AnimSequence->Modify();
 
-			const float StartTime = FilteredFrames.FrameSet->GetAllFrames()[FrameIdx] / FMath::Max(Database->GetFrameRate().AsDecimal(), UE_SMALL_NUMBER);
+				int32 AnimNotifyTrack = AnimSequence->AnimNotifyTracks.IndexOfByPredicate(
+					[&](const FAnimNotifyTrack& Track)
+					{
+						return Track.TrackName == AnimNotifyTrackNames[TrackIdx];
+					});
 
-			FAnimNotifyEvent& NewEvent = AnimSequence->Notifies.AddDefaulted_GetRef();
+				if (AnimNotifyTrack == INDEX_NONE)
+				{
+					AnimNotifyTrack = AnimSequence->AnimNotifyTracks.Add(FAnimNotifyTrack(AnimNotifyTrackNames[TrackIdx], FLinearColor::Red));
+				}
 
-			NewEvent.NotifyName = NAME_None;
-			NewEvent.Link(AnimSequence, StartTime);
-			NewEvent.TriggerTimeOffset = GetTriggerTimeOffsetForType(AnimSequence->CalculateOffsetForNotify(StartTime));
-			NewEvent.TrackIndex = AnimNotifyTrack;
-			NewEvent.Notify = nullptr;
-			NewEvent.Guid = FGuid::NewGuid();
-			NewEvent.Notify = DuplicateObject(AnimNotify, AnimSequence);
-			NewEvent.NotifyName = FName(*NewEvent.Notify->GetNotifyName());
-			NewEvent.EndLink.Link(AnimSequence, NewEvent.EndLink.GetTime());
+				if (Colors.IsValidIndex(TrackIdx))
+				{
+					AnimSequence->AnimNotifyTracks[AnimNotifyTrack].TrackColor = Colors[TrackIdx];
+				}
+
+				const int32 EntryFrameNum = FilteredFrames.FrameSet->GetEntryFrameNum(EntryIdx);
 
-			// Refresh all cached data
-			AnimSequence->RefreshCacheData();
-			AnimSequence->Modify(true);
+				for (int32 FrameIdx = 0; FrameIdx < EntryFrameNum; FrameIdx++)
+				{
+					const float StartTime = FilteredFrames.FrameSet->GetEntryFrame(EntryIdx, FrameIdx) / FMath::Max(Database->GetFrameRate().AsDecimal(), UE_SMALL_NUMBER);
+
+					FAnimNotifyEvent& NewEvent = AnimSequence->Notifies.AddDefaulted_GetRef();
+
+					NewEvent.NotifyName = NAME_None;
+					NewEvent.Link(AnimSequence, StartTime);
+					NewEvent.TriggerTimeOffset = GetTriggerTimeOffsetForType(AnimSequence->CalculateOffsetForNotify(StartTime));
+					NewEvent.TrackIndex = AnimNotifyTrack;
+					NewEvent.Notify = nullptr;
+					NewEvent.Guid = FGuid::NewGuid();
+					NewEvent.Notify = DuplicateObject(AnimNotifies[TrackIdx], AnimSequence);
+					NewEvent.NotifyName = FName(*NewEvent.Notify->GetNotifyName());
+					NewEvent.EndLink.Link(AnimSequence, NewEvent.EndLink.GetTime());
+				}
+
+				AnimSequence->RefreshCacheData();
+				AnimSequence->MarkPackageDirty();
+			}
 		}
 	}
 #endif
 }
 
-void UAnimDatabaseLibrary::DatabaseAddAnimNotifyStateObject(UAnimDatabase* Database, const FAnimDatabaseFrameRanges& FrameRanges, const FName& TrackName, const UAnimNotifyState* AnimNotifyState)
+void UAnimDatabaseLibrary::DatabaseAddAnimNotifyStateObject(UAnimDatabase* Database, const FName TrackName, UAnimNotifyState* AnimNotifyState, const FAnimDatabaseFrameRanges& FrameRanges, const FLinearColor Color, const bool bShouldTransact)
+{
+	DatabaseAddAnimNotifyStateObjectsFromArrayViews(Database, { TrackName }, { AnimNotifyState }, { FrameRanges }, { Color }, bShouldTransact);
+}
+
+void UAnimDatabaseLibrary::DatabaseAddAnimNotifyStateObjects(UAnimDatabase* Database, const TArray<FName>& AnimNotifyStateTrackNames, const TArray<UAnimNotifyState*>& AnimNotifyStates, const TArray<FAnimDatabaseFrameRanges>& FrameRanges, const TArray<FLinearColor>& Colors, const bool bShouldTransact)
+{
+	DatabaseAddAnimNotifyStateObjectsFromArrayViews(Database, AnimNotifyStateTrackNames, AnimNotifyStates, FrameRanges, Colors, bShouldTransact);
+}
+
+void UAnimDatabaseLibrary::DatabaseAddAnimNotifyStateObjectsFromArrayViews(UAnimDatabase* Database, const TArrayView<const FName> AnimNotifyStateTrackNames, const TArrayView<UAnimNotifyState* const> AnimNotifyStates, const TArrayView<const FAnimDatabaseFrameRanges> FrameRanges, const TArrayView<const FLinearColor> Colors, const bool bShouldTransact)
 {
 #if WITH_EDITOR
-	TRACE_CPUPROFILER_EVENT_SCOPE(UAnimDatabaseLibrary::DatabaseAddAnimNotifyStateObject);
+	TRACE_CPUPROFILER_EVENT_SCOPE(UAnimDatabaseLibrary::DatabaseAddAnimNotifyStateObjectsFromArrayViews);
 
 	if (!Database)
 	{
-		UE_LOGFMT(LogAnimDatabase, Error, "DatabaseAddAnimNotifyStateObject: Database is nullptr.");
+		UE_LOGFMT(LogAnimDatabase, Error, "DatabaseAddAnimNotifyStateObjectsFromArrayViews: Database is nullptr.");
 		return;
 	}
 
-	if (!AnimNotifyState)
+	if (AnimNotifyStateTrackNames.Num() != AnimNotifyStates.Num() ||
+		AnimNotifyStateTrackNames.Num() != FrameRanges.Num())
 	{
-		UE_LOGFMT(LogAnimDatabase, Error, "DatabaseAddAnimNotifyStateObject: AnimNotifyState is nullptr.");
+		UE_LOGFMT(LogAnimDatabase, Error, "DatabaseAddAnimNotifyStateObjectsFromArrayViews: Number of Track Names, NotifyStates, and FrameRanges must match. Got {Tracks}, {NotifyStates}, and {Ranges}.", AnimNotifyStateTrackNames.Num(), AnimNotifyStates.Num(), FrameRanges.Num());
 		return;
 	}
 
-	const FAnimDatabaseFrameRanges FilteredFrameRanges = UAnimDatabaseFrameRangesLibrary::MakeFrameRangesFromNotMirrored(Database, FrameRanges);
-	const int32 TotalRangeNum = FilteredFrameRanges.FrameRangeSet->GetTotalRangeNum();
-
-	TLearningArray<1, int32> RangeSequenceIndices;
-	RangeSequenceIndices.SetNumUninitialized({ TotalRangeNum });
-	UE::Learning::FrameRangeSet::AllRangeSequences(RangeSequenceIndices, *FilteredFrameRanges.FrameRangeSet);
+	const int32 TrackNum = AnimNotifyStateTrackNames.Num();
 
-	for (int32 RangeIdx = 0; RangeIdx < TotalRangeNum; RangeIdx++)
+	for (int32 TrackIdx = 0; TrackIdx < TrackNum; TrackIdx++)
 	{
-		TRACE_CPUPROFILER_EVENT_SCOPE(UAnimDatabaseLibrary::DatabaseAddAnimNotifyStatesFromArrayViews::Range);
-		check(!Database->GetIsMirrored(RangeSequenceIndices[RangeIdx]));
+		if (!AnimNotifyStates[TrackIdx])
+		{
+			UE_LOGFMT(LogAnimDatabase, Error, "DatabaseAddAnimNotifyStateObjectsFromArrayViews: AnimNotifyState is nullptr.");
+			continue;
+		}
+
+		const FAnimDatabaseFrameRanges FilteredFrameRanges = UAnimDatabaseFrameRangesLibrary::MakeFrameRangesFromNotMirrored(Database, FrameRanges[TrackIdx]);
+
+		const int32 EntryNum = FilteredFrameRanges.FrameRangeSet->GetEntryNum();
 
-		if (UAnimSequence* AnimSequence = Database->GetAnimSequence(RangeSequenceIndices[RangeIdx]))
+		for (int32 EntryIdx = 0; EntryIdx < EntryNum; EntryIdx++)
 		{
-			int32 AnimNotifyTrack = AnimSequence->AnimNotifyTracks.IndexOfByPredicate(
-				[&](const FAnimNotifyTrack& Track)
-				{
-					return Track.TrackName == TrackName;
-				});
+			check(!Database->GetIsMirrored(FilteredFrameRanges.FrameRangeSet->GetEntrySequence(EntryIdx)));
 
-			if (AnimNotifyTrack == INDEX_NONE)
+			if (UAnimSequence* AnimSequence = Database->GetAnimSequence(FilteredFrameRanges.FrameRangeSet->GetEntrySequence(EntryIdx)))
 			{
-				AnimNotifyTrack = AnimSequence->AnimNotifyTracks.Add(FAnimNotifyTrack(TrackName, FLinearColor::Red));
-			}
+				FScopedTransaction Transaction(LOCTEXT("AddAnimNotifyStateObjectsTransaction", "Adding Anim Notify State Objects"), bShouldTransact);
+				AnimSequence->Modify();
 
-			const float StartTime = FilteredFrameRanges.FrameRangeSet->GetAllRangeStarts()[RangeIdx] / FMath::Max(Database->GetFrameRate().AsDecimal(), UE_SMALL_NUMBER);
-			const float Duration = (FilteredFrameRanges.FrameRangeSet->GetAllRangeLengths()[RangeIdx] - 1) / FMath::Max(Database->GetFrameRate().AsDecimal(), UE_SMALL_NUMBER);
+				int32 AnimNotifyTrack = AnimSequence->AnimNotifyTracks.IndexOfByPredicate(
+					[&](const FAnimNotifyTrack& Track)
+					{
+						return Track.TrackName == AnimNotifyStateTrackNames[TrackIdx];
+					});
 
-			FAnimNotifyEvent& NewEvent = AnimSequence->Notifies.AddDefaulted_GetRef();
+				if (AnimNotifyTrack == INDEX_NONE)
+				{
+					AnimNotifyTrack = AnimSequence->AnimNotifyTracks.Add(FAnimNotifyTrack(AnimNotifyStateTrackNames[TrackIdx], FLinearColor::Red));
+				}
 
-			NewEvent.NotifyName = NAME_None;
-			NewEvent.Link(AnimSequence, StartTime);
-			NewEvent.TriggerTimeOffset = GetTriggerTimeOffsetForType(AnimSequence->CalculateOffsetForNotify(StartTime));
-			NewEvent.TrackIndex = AnimNotifyTrack;
-			NewEvent.Notify = nullptr;
-			NewEvent.Guid = FGuid::NewGuid();
-			NewEvent.NotifyStateClass = DuplicateObject(AnimNotifyState, AnimSequence);
-			NewEvent.NotifyName = FName(*NewEvent.NotifyStateClass->GetNotifyName());
-			NewEvent.SetDuration(Duration);
-			NewEvent.EndLink.Link(AnimSequence, NewEvent.EndLink.GetTime());
+				if (Colors.IsValidIndex(TrackIdx))
+				{
+					AnimSequence->AnimNotifyTracks[AnimNotifyTrack].TrackColor = Colors[TrackIdx];
+				}
+
+				const int32 EntryRangeNum = FilteredFrameRanges.FrameRangeSet->GetEntryRangeNum(EntryIdx);
 
-			// Refresh all cached data
-			AnimSequence->RefreshCacheData();
-			AnimSequence->Modify(true);
+				for (int32 RangeIdx = 0; RangeIdx < EntryRangeNum; RangeIdx++)
+				{
+					const float StartTime = FilteredFrameRanges.FrameRangeSet->GetEntryRangeStart(EntryIdx, RangeIdx) / FMath::Max(Database->GetFrameRate().AsDecimal(), UE_SMALL_NUMBER);
+					const float Duration = (FilteredFrameRanges.FrameRangeSet->GetEntryRangeLength(EntryIdx, RangeIdx) - 1) / FMath::Max(Database->GetFrameRate().AsDecimal(), UE_SMALL_NUMBER);
+
+					FAnimNotifyEvent& NewEvent = AnimSequence->Notifies.AddDefaulted_GetRef();
+
+					NewEvent.NotifyName = NAME_None;
+					NewEvent.Link(AnimSequence, StartTime);
+					NewEvent.TriggerTimeOffset = GetTriggerTimeOffsetForType(AnimSequence->CalculateOffsetForNotify(StartTime));
+					NewEvent.TrackIndex = AnimNotifyTrack;
+					NewEvent.Notify = nullptr;
+					NewEvent.Guid = FGuid::NewGuid();
+					NewEvent.NotifyStateClass = DuplicateObject(AnimNotifyStates[TrackIdx], AnimSequence);
+					NewEvent.NotifyName = FName(*NewEvent.NotifyStateClass->GetNotifyName());
+					NewEvent.SetDuration(Duration);
+					NewEvent.EndLink.Link(AnimSequence, NewEvent.EndLink.GetTime());
+				}
+
+				AnimSequence->RefreshCacheData();
+				AnimSequence->MarkPackageDirty();
+			}
 		}
 	}
 #endif
 }
 
-void UAnimDatabaseLibrary::DatabaseAddAnimNotify(UAnimDatabase* Database, const FName AnimNotifyTrackName, const TSubclassOf<UAnimNotify> Class, const FAnimDatabaseFrames& Frames)
+void UAnimDatabaseLibrary::DatabaseAddAnimNotify(UAnimDatabase* Database, const FName AnimNotifyTrackName, const TSubclassOf<UAnimNotify> Class, const FAnimDatabaseFrames& Frames, const FLinearColor Color, const bool bShouldTransact)
 {
-	DatabaseAddAnimNotifiesFromArrayViews(Database, { AnimNotifyTrackName }, { Class }, { Frames });
+	DatabaseAddAnimNotifiesFromArrayViews(Database, { AnimNotifyTrackName }, { Class }, { Frames }, { Color }, bShouldTransact);
 }
 
-void UAnimDatabaseLibrary::DatabaseAddAnimNotifies(UAnimDatabase* Database, const TArray<FName>& AnimNotifyTrackNames, const TArray<TSubclassOf<UAnimNotify>>& Classes, const TArray<FAnimDatabaseFrames>& Frames)
+void UAnimDatabaseLibrary::DatabaseAddAnimNotifies(UAnimDatabase* Database, const TArray<FName>& AnimNotifyTrackNames, const TArray<TSubclassOf<UAnimNotify>>& Classes, const TArray<FAnimDatabaseFrames>& Frames, const TArray<FLinearColor>& Colors, const bool bShouldTransact)
 {
-	DatabaseAddAnimNotifiesFromArrayViews(Database, AnimNotifyTrackNames, Classes, Frames);
+	DatabaseAddAnimNotifiesFromArrayViews(Database, AnimNotifyTrackNames, Classes, Frames, Colors, bShouldTransact);
 }
 
-void UAnimDatabaseLibrary::DatabaseAddAnimNotifiesFromArrayViews(UAnimDatabase* Database, const TArrayView<const FName> AnimNotifyTrackNames, const TArrayView<const TSubclassOf<UAnimNotify>> Classes, const TArrayView<const FAnimDatabaseFrames> Frames)
+void UAnimDatabaseLibrary::DatabaseAddAnimNotifiesFromArrayViews(UAnimDatabase* Database, const TArrayView<const FName> AnimNotifyTrackNames, const TArrayView<const TSubclassOf<UAnimNotify>> Classes, const TArrayView<const FAnimDatabaseFrames> Frames, const TArrayView<const FLinearColor> Colors, const bool bShouldTransact)
 {
 #if WITH_EDITOR
+
 	TRACE_CPUPROFILER_EVENT_SCOPE(UAnimDatabaseLibrary::DatabaseAddAnimNotifiesFromArrayViews);
 
+	if (!Database)
+	{
+		UE_LOGFMT(LogAnimDatabase, Error, "DatabaseAddAnimNotifiesFromArrayViews: Database is nullptr.");
+		return;
+	}
+
 	if (AnimNotifyTrackNames.Num() != Classes.Num() ||
 		AnimNotifyTrackNames.Num() != Frames.Num())
 	{
@@ -6216,25 +6301,17 @@ void UAnimDatabaseLibrary::DatabaseAddAnimNotifiesFromArrayViews(UAnimDatabase*
 		const FAnimDatabaseFrames FilteredFrames = UAnimDatabaseFrameRangesLibrary::MakeFramesAtFrameRangesStarts(
 			UAnimDatabaseFrameRangesLibrary::MakeFrameRangesFromNotMirrored(Database, UAnimDatabaseFrameRangesLibrary::MakeFrameRangesFromFrames(Frames[TrackIdx])));
 
-		const int32 TotalEventNum = FilteredFrames.FrameSet->GetTotalFrameNum();
-
-		TLearningArray<1, int32> EventSequenceIndices;
-		TLearningArray<1, float> EventStartTimes;
+		const int32 EntryNum = FilteredFrames.FrameSet->GetEntryNum();
 
-		EventSequenceIndices.SetNumUninitialized({ TotalEventNum });
-		EventStartTimes.SetNumUninitialized({ TotalEventNum });
-
-		UE::Learning::FrameSet::AllFrameSequences(EventSequenceIndices, *FilteredFrames.FrameSet);
-		UE::Learning::FrameSet::AllFrameTimes(EventStartTimes, *FilteredFrames.FrameSet, 1.0f / FMath::Max(Database->GetFrameRate().AsDecimal(), UE_SMALL_NUMBER));
-
-		for (int32 EventIdx = 0; EventIdx < TotalEventNum; EventIdx++)
+		for (int32 EntryIdx = 0; EntryIdx < EntryNum; EntryIdx++)
 		{
-			TRACE_CPUPROFILER_EVENT_SCOPE(UAnimDatabaseLibrary::DatabaseAddAnimNotifiesFromArrayViews::Event);
-
-			check(!Database->GetIsMirrored(EventSequenceIndices[EventIdx]));
+			check(!Database->GetIsMirrored(FilteredFrames.FrameSet->GetEntrySequence(EntryIdx)));
 
-			if (UAnimSequence* AnimSequence = Database->GetAnimSequence(EventSequenceIndices[EventIdx]))
+			if (UAnimSequence* AnimSequence = Database->GetAnimSequence(FilteredFrames.FrameSet->GetEntrySequence(EntryIdx)))
 			{
+				FScopedTransaction Transaction(LOCTEXT("AddAnimNotifiesTransaction", "Adding Anim Notifies"), bShouldTransact);
+				AnimSequence->Modify();
+
 				int32 AnimNotifyTrack = AnimSequence->AnimNotifyTracks.IndexOfByPredicate(
 					[&](const FAnimNotifyTrack& Track)
 					{
@@ -6246,37 +6323,133 @@ void UAnimDatabaseLibrary::DatabaseAddAnimNotifiesFromArrayViews(UAnimDatabase*
 					AnimNotifyTrack = AnimSequence->AnimNotifyTracks.Add(FAnimNotifyTrack(AnimNotifyTrackNames[TrackIdx], FLinearColor::Red));
 				}
 
-				FAnimNotifyEvent& NewEvent = AnimSequence->Notifies.AddDefaulted_GetRef();
+				if (Colors.IsValidIndex(TrackIdx))
+				{
+					AnimSequence->AnimNotifyTracks[AnimNotifyTrack].TrackColor = Colors[TrackIdx];
+				}
 
-				NewEvent.NotifyName = NAME_None;
-				NewEvent.Link(AnimSequence, EventStartTimes[EventIdx]);
-				NewEvent.TriggerTimeOffset = GetTriggerTimeOffsetForType(AnimSequence->CalculateOffsetForNotify(EventStartTimes[EventIdx]));
-				NewEvent.TrackIndex = AnimNotifyTrack;
-				NewEvent.NotifyStateClass = nullptr;
-				NewEvent.Guid = FGuid::NewGuid();
+				const int32 EntryFrameNum = FilteredFrames.FrameSet->GetEntryFrameNum(EntryIdx);
 
-				if (Classes[TrackIdx])
+				for (int32 FrameIdx = 0; FrameIdx < EntryFrameNum; FrameIdx++)
 				{
-					NewEvent.Notify = NewObject<UAnimNotify>(AnimSequence, Classes[TrackIdx], NAME_None, RF_Transactional);
-					if (NewEvent.Notify)
+					const float StartTime = FilteredFrames.FrameSet->GetEntryFrame(EntryIdx, FrameIdx) / FMath::Max(Database->GetFrameRate().AsDecimal(), UE_SMALL_NUMBER);
+
+					FAnimNotifyEvent& NewEvent = AnimSequence->Notifies.AddDefaulted_GetRef();
+					NewEvent.NotifyName = NAME_None;
+					NewEvent.Link(AnimSequence, StartTime);
+					NewEvent.TriggerTimeOffset = GetTriggerTimeOffsetForType(AnimSequence->CalculateOffsetForNotify(StartTime));
+					NewEvent.TrackIndex = AnimNotifyTrack;
+					NewEvent.NotifyStateClass = nullptr;
+					NewEvent.Guid = FGuid::NewGuid();
+
+					if (Classes[TrackIdx])
+					{
+						NewEvent.Notify = NewObject<UAnimNotify>(AnimSequence, Classes[TrackIdx], NAME_None, RF_Transactional);
+						if (NewEvent.Notify)
+						{
+							NewEvent.NotifyName = FName(*NewEvent.Notify->GetNotifyName());
+						}
+					}
+					else
 					{
-						NewEvent.NotifyName = FName(*NewEvent.Notify->GetNotifyName());
+						NewEvent.Notify = nullptr;
 					}
 				}
-				else
+
+				AnimSequence->RefreshCacheData();
+				AnimSequence->MarkPackageDirty();
+			}
+		}
+	}
+#endif
+}
+
+void UAnimDatabaseLibrary::DatabaseAddSyncMarker(UAnimDatabase* Database, const FName SyncMarkerTrackName, const FName SyncMarkerName, const FAnimDatabaseFrames& Frames, const FLinearColor Color, const bool bShouldTransact)
+{
+	DatabaseAddSyncMarkersFromArrayViews(Database, { SyncMarkerTrackName }, { SyncMarkerName }, { Frames }, { Color }, bShouldTransact);
+}
+
+void UAnimDatabaseLibrary::DatabaseAddSyncMarkers(UAnimDatabase* Database, const TArray<FName>& SyncMarkerTrackNames, const TArray<FName>& SyncMarkerNames, const TArray<FAnimDatabaseFrames>& Frames, const TArray<FLinearColor>& Colors, const bool bShouldTransact)
+{
+	DatabaseAddSyncMarkersFromArrayViews(Database, SyncMarkerTrackNames, SyncMarkerNames, Frames, Colors, bShouldTransact);
+}
+
+void UAnimDatabaseLibrary::DatabaseAddSyncMarkersFromArrayViews(UAnimDatabase* Database, const TArrayView<const FName> SyncMarkerTrackNames, const TArrayView<const FName> SyncMarkerNames, const TArrayView<const FAnimDatabaseFrames> Frames, const TArrayView<const FLinearColor> Colors, const bool bShouldTransact)
+{
+#if WITH_EDITOR
+
+	TRACE_CPUPROFILER_EVENT_SCOPE(UAnimDatabaseLibrary::DatabaseAddSyncMarkersFromArrayViews);
+
+	if (!Database)
+	{
+		UE_LOGFMT(LogAnimDatabase, Error, "DatabaseAddSyncMarkersFromArrayViews: Database is nullptr.");
+		return;
+	}
+
+	if (SyncMarkerTrackNames.Num() != SyncMarkerNames.Num() ||
+		SyncMarkerTrackNames.Num() != Frames.Num())
+	{
+		UE_LOGFMT(LogAnimDatabase, Error, "DatabaseAddSyncMarkersFromArrayViews: Number of Track Names, Sync Marker names, and Frames must match. Got {Tracks}, {SyncMarkers}, and {Frames}.", SyncMarkerTrackNames.Num(), SyncMarkerNames.Num(), Frames.Num());
+		return;
+	}
+
+	const int32 TrackNum = SyncMarkerTrackNames.Num();
+
+	for (int32 TrackIdx = 0; TrackIdx < TrackNum; TrackIdx++)
+	{
+		const FAnimDatabaseFrames FilteredFrames = UAnimDatabaseFrameRangesLibrary::MakeFramesAtFrameRangesStarts(
+			UAnimDatabaseFrameRangesLibrary::MakeFrameRangesFromNotMirrored(Database, UAnimDatabaseFrameRangesLibrary::MakeFrameRangesFromFrames(Frames[TrackIdx])));
+
+		const int32 EntryNum = FilteredFrames.FrameSet->GetEntryNum();
+
+		for (int32 EntryIdx = 0; EntryIdx < EntryNum; EntryIdx++)
+		{
+			check(!Database->GetIsMirrored(FilteredFrames.FrameSet->GetEntrySequence(EntryIdx)));
+
+			if (UAnimSequence* AnimSequence = Database->GetAnimSequence(FilteredFrames.FrameSet->GetEntrySequence(EntryIdx)))
+			{
+				FScopedTransaction Transaction(LOCTEXT("AddSyncMarkersTransaction", "Adding Sync Markers"), bShouldTransact);
+				AnimSequence->Modify();
+
+				int32 AnimNotifyTrack = AnimSequence->AnimNotifyTracks.IndexOfByPredicate(
+					[&](const FAnimNotifyTrack& Track)
+					{
+						return Track.TrackName == SyncMarkerTrackNames[TrackIdx];
+					});
+
+				if (AnimNotifyTrack == INDEX_NONE)
 				{
-					NewEvent.Notify = nullptr;
+					AnimNotifyTrack = AnimSequence->AnimNotifyTracks.Add(FAnimNotifyTrack(SyncMarkerTrackNames[TrackIdx], FLinearColor::Red));
+				}
+
+				if (Colors.IsValidIndex(TrackIdx))
+				{
+					AnimSequence->AnimNotifyTracks[AnimNotifyTrack].TrackColor = Colors[TrackIdx];
+				}
+
+				const int32 EntryFrameNum = FilteredFrames.FrameSet->GetEntryFrameNum(EntryIdx);
+
+				for (int32 FrameIdx = 0; FrameIdx < EntryFrameNum; FrameIdx++)
+				{
+					const float StartTime = FilteredFrames.FrameSet->GetEntryFrame(EntryIdx, FrameIdx) / FMath::Max(Database->GetFrameRate().AsDecimal(), UE_SMALL_NUMBER);
+
+					FAnimSyncMarker& Marker = AnimSequence->AuthoredSyncMarkers.AddDefaulted_GetRef();
+					Marker.MarkerName = SyncMarkerNames[TrackIdx];
+					Marker.Time = StartTime;
+					Marker.TrackIndex = AnimNotifyTrack;
+					Marker.Guid = FGuid::NewGuid();
 				}
 
-				// Refresh all cached data
 				AnimSequence->RefreshCacheData();
-				AnimSequence->Modify(true);
+				AnimSequence->MarkPackageDirty();
 			}
 		}
 	}
 #endif
 }
 
+
+
 void UAnimDatabaseLibrary::DatabaseSetLocalBoneTransform(UAnimDatabase* Database, const FName BoneName, const FAnimDatabaseFrameAttribute& TransformFrameAttribute, const bool bShouldTransact)
 {
 	DatabaseSetLocalBoneTransformsFromArrayViews(Database, { BoneName }, { TransformFrameAttribute }, bShouldTransact);
@@ -7061,7 +7234,7 @@ void UAnimDatabaseFunction_BoneApplyTransform::Run_Implementation(UAnimDatabase*
 
 void UAnimDatabaseFunction_AddAnimNotifyState::Run_Implementation(UAnimDatabase* InDatabase, const FAnimDatabaseFrameRanges& InFrameRanges, const bool bShouldTransact)
 {
-	UAnimDatabaseLibrary::DatabaseAddAnimNotifyStateObject(InDatabase, InFrameRanges, TrackName, AnimNotifyState);
+	UAnimDatabaseLibrary::DatabaseAddAnimNotifyStateObject(InDatabase, TrackName, AnimNotifyState, InFrameRanges, Color, bShouldTransact);
 }
 
 void UAnimDatabaseFunction_Sequence::Run_Implementation(UAnimDatabase* InDatabase, const FAnimDatabaseFrameRanges& InFrameRanges, const bool bShouldTransact)
@@ -7079,18 +7252,22 @@ void UAnimDatabaseFunction_AddAnimNotifyAtFrames::Run_Implementation(UAnimDataba
 {
 	UAnimDatabaseLibrary::DatabaseAddAnimNotifyObject(
 		InDatabase,
-		UAnimDatabaseFramesLibrary::MakeFramesFromFunction(InDatabase, InFrameRanges, Frames),
 		TrackName,
-		AnimNotify);
+		AnimNotify, 
+		UAnimDatabaseFramesLibrary::MakeFramesFromFunction(InDatabase, InFrameRanges, Frames),
+		Color,
+		bShouldTransact);
 }
 
 void UAnimDatabaseFunction_AddAnimNotifyStateAtFrameRanges::Run_Implementation(UAnimDatabase* InDatabase, const FAnimDatabaseFrameRanges& InFrameRanges, const bool bShouldTransact)
 {
 	UAnimDatabaseLibrary::DatabaseAddAnimNotifyStateObject(
 		InDatabase,
-		UAnimDatabaseFrameRangesLibrary::MakeFrameRangesFromFunction(InDatabase, InFrameRanges, FrameRanges),
 		TrackName,
-		AnimNotifyState);
+		AnimNotifyState,
+		UAnimDatabaseFrameRangesLibrary::MakeFrameRangesFromFunction(InDatabase, InFrameRanges, FrameRanges),
+		Color,
+		bShouldTransact);
 }
 
 void UAnimDatabaseFunction_AddCurveFromFrameAttribute::Run_Implementation(UAnimDatabase* InDatabase, const FAnimDatabaseFrameRanges& InFrameRanges, const bool bShouldTransact)
@@ -7099,47 +7276,124 @@ void UAnimDatabaseFunction_AddCurveFromFrameAttribute::Run_Implementation(UAnimD
 		InDatabase, 
 		CurveName, 
 		UAnimDatabaseFrameAttributeLibrary::MakeFrameAttributeFromFunction(InDatabase, InFrameRanges, FrameAttribute),
+		Color,
 		bSparseKeys, 
 		bShouldTransact);
 }
 
-void  UAnimDatabaseFunction_AddContactCurves::Run_Implementation(UAnimDatabase* InDatabase, const FAnimDatabaseFrameRanges& InFrameRanges, const bool bShouldTransact)
+void UAnimDatabaseFunction_ProcessContacts::Run_Implementation(UAnimDatabase* InDatabase, const FAnimDatabaseFrameRanges& InFrameRanges, const bool bShouldTransact)
 {
 	if (!InDatabase) { return; }
 	if (!InFrameRanges.IsValid()) { return; }
 
 	InDatabase->WaitForCompressionOnAnimSequencesFromArrayView(InFrameRanges.FrameRangeSet->GetEntrySequences());
 
-	UE::Tasks::TTask<FAnimDatabaseFrameAttribute> LeftContactCurve = UE::Tasks::Launch(UE_SOURCE_LOCATION, [this, InDatabase, InFrameRanges]() {
-		
-		return UAnimDatabaseFrameAttributeLibrary::FrameAttributeContactCurve(
-			InDatabase,
-			InFrameRanges,
-			LeftContactBoneName,
-			HeightThreshold,
-			VelocityThreshold,
-			MajorityVoteFilterWidth,
-			SmoothingAmount);
-		});
+	UE::Tasks::TTask<FAnimDatabaseFrameAttribute> LeftFootSpeedCurve, RightFootSpeedCurve;
+
+	if (bAddFootSpeedCurves)
+	{
+		LeftFootSpeedCurve = UE::Tasks::Launch(UE_SOURCE_LOCATION, [this, InDatabase, InFrameRanges]() {
 
-	UE::Tasks::TTask<FAnimDatabaseFrameAttribute> RightContactCurve = UE::Tasks::Launch(UE_SOURCE_LOCATION, [this, InDatabase, InFrameRanges]() {
+			return UAnimDatabaseFrameAttributeLibrary::FrameAttributeBoneSpeed(
+				InDatabase,
+				InFrameRanges,
+				LeftFootSpeedBoneName);
+			});
 
-		return UAnimDatabaseFrameAttributeLibrary::FrameAttributeContactCurve(
+		RightFootSpeedCurve = UE::Tasks::Launch(UE_SOURCE_LOCATION, [this, InDatabase, InFrameRanges]() {
+
+			return UAnimDatabaseFrameAttributeLibrary::FrameAttributeBoneSpeed(
+				InDatabase,
+				InFrameRanges,
+				RightFootSpeedBoneName);
+			});
+	}
+
+	UE::Tasks::TTask<FAnimDatabaseFrameAttribute> LeftContactCurve, RightContactCurve;
+
+	if (bAddContactCurves || bAddAnimNotifies || bAddSyncMarkers)
+	{
+		LeftContactCurve = UE::Tasks::Launch(UE_SOURCE_LOCATION, [this, InDatabase, InFrameRanges]() {
+
+			return UAnimDatabaseFrameAttributeLibrary::FrameAttributeContactCurve(
+				InDatabase,
+				InFrameRanges,
+				LeftContactBoneName,
+				HeightThreshold,
+				VelocityThreshold,
+				bFilterContactCurves,
+				MajorityVoteFilterWidth,
+				bSmoothContactCurves,
+				SmoothingAmount);
+			});
+
+		RightContactCurve = UE::Tasks::Launch(UE_SOURCE_LOCATION, [this, InDatabase, InFrameRanges]() {
+
+			return UAnimDatabaseFrameAttributeLibrary::FrameAttributeContactCurve(
+				InDatabase,
+				InFrameRanges,
+				RightContactBoneName,
+				HeightThreshold,
+				VelocityThreshold,
+				bFilterContactCurves,
+				MajorityVoteFilterWidth,
+				bSmoothContactCurves,
+				SmoothingAmount);
+			});
+	}
+
+	if (bAddFootSpeedCurves)
+	{
+		UAnimDatabaseLibrary::DatabaseAddCurvesFromArrayViews(
 			InDatabase,
-			InFrameRanges,
-			RightContactBoneName,
-			HeightThreshold,
-			VelocityThreshold,
-			MajorityVoteFilterWidth,
-			SmoothingAmount);
-		});
+			{ LeftFootSpeedCurveName, RightFootSpeedCurveName },
+			{ LeftFootSpeedCurve.GetResult(), RightFootSpeedCurve.GetResult() },
+			{ LeftFootSpeedCurveColor, RightFootSpeedCurveColor },
+			bSparseKeys,
+			bShouldTransact);
+	}
 
-	UAnimDatabaseLibrary::DatabaseAddCurvesFromArrayViews(
-		InDatabase,
-		{ LeftContactCurveName, RightContactCurveName },
-		{ LeftContactCurve.GetResult(), RightContactCurve.GetResult() },
-		bSparseKeys,
-		bShouldTransact);
+	if (bAddContactCurves)
+	{
+		UAnimDatabaseLibrary::DatabaseAddCurvesFromArrayViews(
+			InDatabase,
+			{ LeftContactCurveName, RightContactCurveName },
+			{ LeftContactCurve.GetResult(), RightContactCurve.GetResult() },
+			{ LeftContactCurveColor, RightContactCurveColor },
+			bSparseKeys,
+			bShouldTransact);
+	}
+
+	if (bAddAnimNotifies || bAddSyncMarkers)
+	{
+		const FAnimDatabaseFrames LeftNotifyFrames = UAnimDatabaseFramesLibrary::FramesBefore(InDatabase, UAnimDatabaseFrameRangesLibrary::MakeFramesAtFrameRangesStarts(UAnimDatabaseFrameAttributeLibrary::FrameAttributeBoolToFrameRanges(
+			UAnimDatabaseFrameAttributeLibrary::FrameAttributeLessThanFloat(LeftContactCurve.GetResult(), EventThreshold))), FMath::RoundToInt(InDatabase->GetFrameRate().AsDecimal() * EventOffset), EAnimDatabaseFrameShiftBehavior::Clamp);
+
+		const FAnimDatabaseFrames RightNotifyFrames = UAnimDatabaseFramesLibrary::FramesBefore(InDatabase, UAnimDatabaseFrameRangesLibrary::MakeFramesAtFrameRangesStarts(UAnimDatabaseFrameAttributeLibrary::FrameAttributeBoolToFrameRanges(
+			UAnimDatabaseFrameAttributeLibrary::FrameAttributeLessThanFloat(RightContactCurve.GetResult(), EventThreshold))), FMath::RoundToInt(InDatabase->GetFrameRate().AsDecimal() * EventOffset), EAnimDatabaseFrameShiftBehavior::Clamp);
+
+		if (bAddAnimNotifies)
+		{
+			UAnimDatabaseLibrary::DatabaseAddAnimNotifyObjectsFromArrayViews(
+				InDatabase,
+				{ LeftContactAnimNotifyTrackName, RightContactAnimNotifyTrackName },
+				{ LeftContactAnimNotify, RightContactAnimNotify },
+				{ LeftNotifyFrames, RightNotifyFrames },
+				{ LeftContactAnimNotifyColor, RightContactAnimNotifyColor },
+				bShouldTransact);
+		}
+
+		if (bAddSyncMarkers)
+		{
+			UAnimDatabaseLibrary::DatabaseAddSyncMarkersFromArrayViews(
+				InDatabase,
+				{ LeftContactSyncMarkerTrackName, RightContactSyncMarkerTrackName },
+				{ LeftContactSyncMarker, RightContactSyncMarker },
+				{ LeftNotifyFrames, RightNotifyFrames },
+				{ LeftContactSyncMarkerColor, RightContactSyncMarkerColor },
+				bShouldTransact);
+		}
+	}
 }
 
 void UAnimDatabaseFunction_SetBoneLocalTransformFromFrameAttribute::Run_Implementation(UAnimDatabase* InDatabase, const FAnimDatabaseFrameRanges& InFrameRanges, const bool bShouldTransact)
diff --git a/Engine/Plugins/Experimental/Animation/AnimDatabase/Source/AnimDatabase/Private/AnimDatabaseFrameAttribute.cpp b/Engine/Plugins/Experimental/Animation/AnimDatabase/Source/AnimDatabase/Private/AnimDatabaseFrameAttribute.cpp
--- a/Engine/Plugins/Experimental/Animation/AnimDatabase/Source/AnimDatabase/Private/AnimDatabaseFrameAttribute.cpp
+++ b/Engine/Plugins/Experimental/Animation/AnimDatabase/Source/AnimDatabase/Private/AnimDatabaseFrameAttribute.cpp
@@ -10050,17 +10050,31 @@ FAnimDatabaseFrameAttribute UAnimDatabaseFrameAttributeLibrary::FrameAttributeCo
 	const FName BoneName,
 	const float HeightThreshold,
 	const float VelocityThreshold,
+	const bool bFilterCurve,
 	const int32 MajorityVoteFilterWidth,
+	const bool bSmoothCurve,
 	const float SmoothingAmount)
 {
+	if (!Database)
+	{
+		UE_LOGFMT(LogAnimDatabase, Error, "FrameAttributeContactCurve: Database is nullptr.");
+		return FAnimDatabaseFrameAttribute();
+	}
+
 	FAnimDatabaseFrameAttribute Location, Velocity;
 	MakeBoneGlobalLocationAndLinearVelocityFrameAttributeFromName(Location, Velocity, Database, FrameRanges, BoneName);
 	
-	const FAnimDatabaseFrameAttribute ContactActive = FrameAttributeAnd(
+	FAnimDatabaseFrameAttribute ContactActive = FrameAttributeAnd(
 		FrameAttributeLessThanFloat(FrameAttributeZ(Location), HeightThreshold),
 		FrameAttributeLessThanFloat(FrameAttributeLength(Velocity), VelocityThreshold));
 
-	return FrameAttributeFilterGaussian(FrameAttributeBoolToFloat(FrameAttributeFilterMajorityVote(ContactActive, MajorityVoteFilterWidth)), SmoothingAmount);
+	if (bFilterCurve) { ContactActive = FrameAttributeFilterMajorityVote(ContactActive, MajorityVoteFilterWidth); }
+
+	ContactActive = FrameAttributeBoolToFloat(ContactActive);
+
+	if (bSmoothCurve) { ContactActive = FrameAttributeFilterGaussian(ContactActive, SmoothingAmount); }
+
+	return ContactActive;
 }
 
 FAnimDatabaseFrames UAnimDatabaseFrameAttributeLibrary::FrameAttributeContactFrames(
@@ -10069,18 +10083,25 @@ FAnimDatabaseFrames UAnimDatabaseFrameAttributeLibrary::FrameAttributeContactFra
 	const FName BoneName,
 	const float HeightThreshold,
 	const float VelocityThreshold,
+	const bool bFilterCurve,
 	const int32 MajorityVoteFilterWidth)
 {
+	if (!Database)
+	{
+		UE_LOGFMT(LogAnimDatabase, Error, "FrameAttributeContactFrames: Database is nullptr.");
+		return FAnimDatabaseFrames();
+	}
+
 	FAnimDatabaseFrameAttribute Location, Velocity;
 	MakeBoneGlobalLocationAndLinearVelocityFrameAttributeFromName(Location, Velocity, Database, FrameRanges, BoneName);
 
-	const FAnimDatabaseFrameAttribute ContactActive = FrameAttributeAnd(
+	FAnimDatabaseFrameAttribute ContactActive = FrameAttributeAnd(
 		FrameAttributeLessThanFloat(FrameAttributeZ(Location), HeightThreshold),
 		FrameAttributeLessThanFloat(FrameAttributeLength(Velocity), VelocityThreshold));
 
-	const FAnimDatabaseFrameRanges ContactRanges = FrameAttributeBoolToFrameRanges(FrameAttributeFilterMajorityVote(ContactActive, MajorityVoteFilterWidth));
+	if (bFilterCurve) { ContactActive = FrameAttributeFilterMajorityVote(ContactActive, MajorityVoteFilterWidth); }
 
-	return UAnimDatabaseFrameRangesLibrary::MakeFramesAtFrameRangesStarts(ContactRanges);
+	return UAnimDatabaseFrameRangesLibrary::MakeFramesAtFrameRangesStarts(FrameAttributeBoolToFrameRanges(ContactActive));
 }
 
 FAnimDatabaseFrameAttribute UAnimDatabaseFrameAttributeLibrary::FrameAttributeContactEvent(
@@ -10089,15 +10110,24 @@ FAnimDatabaseFrameAttribute UAnimDatabaseFrameAttributeLibrary::FrameAttributeCo
 	const FName BoneName,
 	const float HeightThreshold,
 	const float VelocityThreshold,
-	const int32 MajorityVoteFilterWidth)
+	const bool bFilterCurve,
+	const int32 MajorityVoteFilterWidth,
+	const float EventOffset)
 {
-	return MakeEventFrameAttribute(FrameAttributeContactFrames(
+	if (!Database)
+	{
+		UE_LOGFMT(LogAnimDatabase, Error, "FrameAttributeContactEvent: Database is nullptr.");
+		return FAnimDatabaseFrameAttribute();
+	}
+
+	return MakeEventFrameAttribute(UAnimDatabaseFramesLibrary::FramesBefore(Database, FrameAttributeContactFrames(
 			Database,
 			FrameRanges,
 			BoneName,
 			HeightThreshold,
 			VelocityThreshold,
-			MajorityVoteFilterWidth), 
+			bFilterCurve,
+			MajorityVoteFilterWidth), FMath::RoundToInt(EventOffset * Database->GetFrameRate().AsDecimal()), EAnimDatabaseFrameShiftBehavior::Clamp),
 		FrameRanges,
 		Database->GetFrameRate());
 }
@@ -10258,10 +10288,25 @@ FAnimDatabaseFrameAttribute UAnimDatabaseFrameAttributeFunction_ContactCurve::Ma
 		BoneName,
 		HeightThreshold,
 		VelocityThreshold,
+		bFilterCurve,
 		MajorityVoteFilterWidth,
+		bSmoothCurve,
 		SmoothingAmount);
 }
 
+FAnimDatabaseFrameAttribute UAnimDatabaseFrameAttributeFunction_ContactEvent::MakeFrameAttribute_Implementation(const UAnimDatabase* Database, const FAnimDatabaseFrameRanges& FrameRanges) const
+{
+	return UAnimDatabaseFrameAttributeLibrary::FrameAttributeContactEvent(
+		Database,
+		FrameRanges,
+		BoneName,
+		HeightThreshold,
+		VelocityThreshold,
+		bFilterCurve,
+		MajorityVoteFilterWidth,
+		EventOffset);
+}
+
 FAnimDatabaseFrameAttribute UAnimDatabaseFrameAttributeFunction_FitTransform::MakeFrameAttribute_Implementation(const UAnimDatabase* Database, const FAnimDatabaseFrameRanges& FrameRanges) const
 {
 	TArray<FAnimDatabaseFrameAttribute> LocationAttributes, ForwardAttributes, RightAttributes, UpAttributes, LocationWeightAttributes, ForwardWeightAttributes, RightWeightAttributes, UpWeightAttributes;
diff --git a/Engine/Plugins/Experimental/Animation/AnimDatabase/Source/AnimDatabase/Private/AnimDatabaseFrameRanges.cpp b/Engine/Plugins/Experimental/Animation/AnimDatabase/Source/AnimDatabase/Private/AnimDatabaseFrameRanges.cpp
--- a/Engine/Plugins/Experimental/Animation/AnimDatabase/Source/AnimDatabase/Private/AnimDatabaseFrameRanges.cpp
+++ b/Engine/Plugins/Experimental/Animation/AnimDatabase/Source/AnimDatabase/Private/AnimDatabaseFrameRanges.cpp
@@ -283,7 +283,7 @@ FAnimDatabaseFrames UAnimDatabaseFramesLibrary::MakeFramesFromAnimNotify(const U
 
 				if (AnimNotifyStart == PreviousFramesTime)
 				{
-					UE_LOGFMT(LogAnimDatabase, Warning, "Anim Notify {Notify} is at same time as Previous in Anim Sequence {Name} Time {Time} Frame {Frame}", *Notify->GetName(), *Sequence->GetName(), AnimNotifyRef.GetNotify()->GetTime(), AnimNotifyStart);
+					UE_LOGFMT(LogAnimDatabase, Warning, "MakeFramesFromAnimNotify: Anim Notify {Notify} is at same time as Previous in Anim Sequence {Name} Time {Time} Frame {Frame}", *Notify->GetName(), *Sequence->GetName(), AnimNotifyRef.GetNotify()->GetTime(), AnimNotifyStart);
 				}
 				else
 				{
@@ -303,6 +303,73 @@ FAnimDatabaseFrames UAnimDatabaseFramesLibrary::MakeFramesFromAnimNotify(const U
 	return Out;
 }
 
+FAnimDatabaseFrames UAnimDatabaseFramesLibrary::MakeFramesFromSyncMarker(const UAnimDatabase* Database, const FAnimDatabaseFrameRanges& FrameRanges, const FName SyncMarkerName)
+{
+	TRACE_CPUPROFILER_EVENT_SCOPE(UAnimDatabaseFramesLibrary::MakeFramesFromSyncMarker);
+
+	if (!Database)
+	{
+		UE_LOGFMT(LogAnimDatabase, Error, "MakeFramesFromSyncMarker: Database is nullptr.");
+		return FAnimDatabaseFrames();
+	}
+
+	if (!FrameRanges.IsValid())
+	{
+		UE_LOGFMT(LogAnimDatabase, Error, "MakeFramesFromSyncMarker: Invalid FrameRanges Object.");
+		return FAnimDatabaseFrames();
+	}
+
+	FAnimDatabaseFrames Out = MakeEmptyFrames();
+
+	TArray<FAnimSyncMarker, TInlineAllocator<64>> SortedSyncMarkers;
+	TArray<int32> FramesToAdd;
+
+	const int32 EntryNum = FrameRanges.FrameRangeSet->GetEntryNum();
+
+	for (int32 EntryIdx = 0; EntryIdx < EntryNum; EntryIdx++)
+	{
+		const int32 SequenceIdx = FrameRanges.FrameRangeSet->GetEntrySequence(EntryIdx);
+
+		const UAnimSequence* Sequence = Database->GetAnimSequence(SequenceIdx);
+		if (!Sequence) { continue; }
+
+		const int32 SequenceFrameNum = FMath::Max(Database->GetSequenceFrameNum(SequenceIdx), 1);
+
+		SortedSyncMarkers = Sequence->AuthoredSyncMarkers;
+		SortedSyncMarkers.Sort([](const FAnimSyncMarker& A, const FAnimSyncMarker& B) { return A.Time < B.Time; });
+
+		int32 PreviousFramesTime = -1;
+
+		FramesToAdd.Reset();
+		for (const FAnimSyncMarker& SyncMarker : SortedSyncMarkers)
+		{
+			if (SyncMarker.MarkerName == SyncMarkerName)
+			{
+				const int32 SyncMarkerStart = FMath::Clamp(
+					FMath::RoundToInt(SyncMarker.Time * Database->GetFrameRate().AsDecimal()),
+					0, SequenceFrameNum - 1);
+
+				if (SyncMarkerStart == PreviousFramesTime)
+				{
+					UE_LOGFMT(LogAnimDatabase, Warning, "MakeFramesFromSyncMarker: Sync Marker {SyncMarkerName} is at same time as Previous in Anim Sequence {Name} Time {Time} Frame {Frame}", *SyncMarkerName.ToString(), *Sequence->GetName(), SyncMarker.Time, SyncMarkerStart);
+				}
+				else
+				{
+					if (FrameRanges.FrameRangeSet->EntryContains(EntryIdx, SyncMarkerStart))
+					{
+						FramesToAdd.Add(SyncMarkerStart);
+					}
+
+					PreviousFramesTime = SyncMarkerStart;
+				}
+			}
+		}
+
+		Out.FrameSet->AddEntry(SequenceIdx, FramesToAdd);
+	}
+
+	return Out;
+}
 
 FAnimDatabaseFrames UAnimDatabaseFramesLibrary::MakeFramesFromAnimNotifyStateStart(const UAnimDatabase* Database, const FAnimDatabaseFrameRanges& FrameRanges, const TSubclassOf<UAnimNotifyState> State)
 {
@@ -357,7 +424,7 @@ FAnimDatabaseFrames UAnimDatabaseFramesLibrary::MakeFramesFromAnimNotifyStateSta
 
 				if (AnimNotifyStart == PreviousFramesTime)
 				{
-					UE_LOGFMT(LogAnimDatabase, Warning, "Anim Notify State {State} Start is at same time as Previous in Anim Sequence {Name} Time {Time} Frame {Frame}", *State->GetName(), *Sequence->GetName(), AnimNotifyRef.GetNotify()->GetTime(), AnimNotifyStart);
+					UE_LOGFMT(LogAnimDatabase, Warning, "MakeFramesFromAnimNotifyStateStart: Anim Notify State {State} Start is at same time as Previous in Anim Sequence {Name} Time {Time} Frame {Frame}", *State->GetName(), *Sequence->GetName(), AnimNotifyRef.GetNotify()->GetTime(), AnimNotifyStart);
 				}
 				else
 				{
@@ -433,7 +500,7 @@ FAnimDatabaseFrames UAnimDatabaseFramesLibrary::MakeFramesFromAnimNotifyStateEnd
 
 				if (AnimNotifyStop == PreviousFramesTime)
 				{
-					UE_LOGFMT(LogAnimDatabase, Warning, "Anim Notify State {State} End is at same time as Previous in Anim Sequence {Name} Time {Time} Frame {Frame}", *State->GetName(), *Sequence->GetName(), AnimNotifyRef.GetNotify()->GetTime(), AnimNotifyStop);
+					UE_LOGFMT(LogAnimDatabase, Warning, "MakeFramesFromAnimNotifyStateEnd: Anim Notify State {State} End is at same time as Previous in Anim Sequence {Name} Time {Time} Frame {Frame}", *State->GetName(), *Sequence->GetName(), AnimNotifyRef.GetNotify()->GetTime(), AnimNotifyStop);
 				}
 				else
 				{
@@ -510,7 +577,7 @@ bool UAnimDatabaseFramesLibrary::FramesIsEmpty(const FAnimDatabaseFrames& Frames
 	return Frames.FrameSet->IsEmpty();
 }
 
-FAnimDatabaseFrames UAnimDatabaseFramesLibrary::FramesBefore(const UAnimDatabase* Database, const FAnimDatabaseFrames& Frames, const int32 FramesBefore)
+FAnimDatabaseFrames UAnimDatabaseFramesLibrary::FramesBefore(const UAnimDatabase* Database, const FAnimDatabaseFrames& Frames, const int32 FramesBefore, const EAnimDatabaseFrameShiftBehavior ShiftBehavior)
 {
 	if (!Database)
 	{
@@ -524,6 +591,8 @@ FAnimDatabaseFrames UAnimDatabaseFramesLibrary::FramesBefore(const UAnimDatabase
 		return FAnimDatabaseFrames();
 	}
 
+	if (FramesBefore == 0) { return Frames; }
+
 	FAnimDatabaseFrames Out = MakeEmptyFrames();
 	
 	TArray<int32, TInlineAllocator<32>> AddedFrames;
@@ -540,9 +609,25 @@ FAnimDatabaseFrames UAnimDatabaseFramesLibrary::FramesBefore(const UAnimDatabase
 		for (int32 FrameIdx = 0; FrameIdx < FrameNum; FrameIdx++)
 		{
 			const int32 Frame = Frames.FrameSet->GetEntryFrame(EntryIdx, FrameIdx) - FramesBefore;
-			if (Frame >= 0 && Frame < SequenceFrameNum)
+			switch (ShiftBehavior)
+			{
+			case EAnimDatabaseFrameShiftBehavior::Remove:
 			{
-				AddedFrames.Add(Frame);
+				if (Frame >= 0 && Frame < SequenceFrameNum)
+				{
+					AddedFrames.Add(Frame);
+				}
+				break;
+			}
+			case EAnimDatabaseFrameShiftBehavior::Clamp:
+			{
+				AddedFrames.Add(FMath::Clamp(Frame, 0, FMath::Max(SequenceFrameNum - 1, 0)));
+				break;
+			}
+			default:
+			{
+				checkNoEntry();
+			}
 			}
 		}
 
@@ -554,11 +639,11 @@ FAnimDatabaseFrames UAnimDatabaseFramesLibrary::FramesBefore(const UAnimDatabase
 	return Out;
 }
 
-FAnimDatabaseFrames UAnimDatabaseFramesLibrary::FramesAfter(const UAnimDatabase* Database, const FAnimDatabaseFrames& Frames, const int32 FramesAfter)
+FAnimDatabaseFrames UAnimDatabaseFramesLibrary::FramesAfter(const UAnimDatabase* Database, const FAnimDatabaseFrames& Frames, const int32 FramesAfter, const EAnimDatabaseFrameShiftBehavior ShiftBehavior)
 {
 	if (!Database)
 	{
-		UE_LOGFMT(LogAnimDatabase, Error, "FramesBefore: Database is nullptr.");
+		UE_LOGFMT(LogAnimDatabase, Error, "FramesAfter: Database is nullptr.");
 		return FAnimDatabaseFrames();
 	}
 
@@ -568,7 +653,7 @@ FAnimDatabaseFrames UAnimDatabaseFramesLibrary::FramesAfter(const UAnimDatabase*
 		return FAnimDatabaseFrames();
 	}
 
-	return FramesBefore(Database, Frames, -FramesAfter);
+	return FramesBefore(Database, Frames, -FramesAfter, ShiftBehavior);
 }
 
 bool UAnimDatabaseFramesLibrary::FramesEqual(const FAnimDatabaseFrames& A, const FAnimDatabaseFrames& B)
@@ -1055,6 +1140,82 @@ FAnimDatabaseFrameRanges UAnimDatabaseFrameRangesLibrary::MakeFrameRangesFromFor
 	return Out;
 }
 
+FAnimDatabaseFrameRanges UAnimDatabaseFrameRangesLibrary::MakeFrameRangesFromNotLooped(const UAnimDatabase* Database, const FAnimDatabaseFrameRanges& FrameRanges)
+{
+	if (!Database)
+	{
+		UE_LOGFMT(LogAnimDatabase, Error, "MakeFrameRangesFromNotLooped: Database is nullptr.");
+		return FAnimDatabaseFrameRanges();
+	}
+
+	if (!FrameRanges.IsValid())
+	{
+		UE_LOGFMT(LogAnimDatabase, Error, "MakeFrameRangesFromNotLooped: Invalid FrameRanges Object.");
+		return FAnimDatabaseFrameRanges();
+	}
+
+	FAnimDatabaseFrameRanges Out = MakeEmptyFrameRanges();
+
+	const int32 EntryNum = FrameRanges.FrameRangeSet->GetEntryNum();
+
+	for (int32 EntryIdx = 0; EntryIdx < EntryNum; EntryIdx++)
+	{
+		const int32 SequenceIdx = FrameRanges.FrameRangeSet->GetEntrySequence(EntryIdx);
+
+		if (const UAnimSequence* AnimSequence = Database->GetAnimSequence(SequenceIdx))
+		{
+			if (!AnimSequence->bLoop)
+			{
+				Out.FrameRangeSet->AddEntry(
+					SequenceIdx,
+					FrameRanges.FrameRangeSet->GetEntryRangeStarts(EntryIdx),
+					FrameRanges.FrameRangeSet->GetEntryRangeLengths(EntryIdx));
+			}
+		}
+	}
+
+	Out.FrameRangeSet->Check();
+	return Out;
+}
+
+FAnimDatabaseFrameRanges UAnimDatabaseFrameRangesLibrary::MakeFrameRangesFromLooped(const UAnimDatabase* Database, const FAnimDatabaseFrameRanges& FrameRanges)
+{
+	if (!Database)
+	{
+		UE_LOGFMT(LogAnimDatabase, Error, "MakeFrameRangesFromLooped: Database is nullptr.");
+		return FAnimDatabaseFrameRanges();
+	}
+
+	if (!FrameRanges.IsValid())
+	{
+		UE_LOGFMT(LogAnimDatabase, Error, "MakeFrameRangesFromLooped: Invalid FrameRanges Object.");
+		return FAnimDatabaseFrameRanges();
+	}
+
+	FAnimDatabaseFrameRanges Out = MakeEmptyFrameRanges();
+
+	const int32 EntryNum = FrameRanges.FrameRangeSet->GetEntryNum();
+
+	for (int32 EntryIdx = 0; EntryIdx < EntryNum; EntryIdx++)
+	{
+		const int32 SequenceIdx = FrameRanges.FrameRangeSet->GetEntrySequence(EntryIdx);
+
+		if (const UAnimSequence* AnimSequence = Database->GetAnimSequence(SequenceIdx))
+		{
+			if (AnimSequence->bLoop)
+			{
+				Out.FrameRangeSet->AddEntry(
+					SequenceIdx,
+					FrameRanges.FrameRangeSet->GetEntryRangeStarts(EntryIdx),
+					FrameRanges.FrameRangeSet->GetEntryRangeLengths(EntryIdx));
+			}
+		}
+	}
+
+	Out.FrameRangeSet->Check();
+	return Out;
+}
+
 FAnimDatabaseFrameRanges UAnimDatabaseFrameRangesLibrary::MakeFrameRangesFromSequenceIndex(const UAnimDatabase* Database, const int32 SequenceIdx)
 {
 	if (!Database)
@@ -1623,7 +1784,7 @@ FAnimDatabaseFrameRanges UAnimDatabaseFrameRangesLibrary::MakeFrameRangesFromAni
 
 				if (AnimNotifyStart < PrevAnimNotifyStopFrame)
 				{
-					UE_LOGFMT(LogAnimDatabase, Warning, "Anim Notify State {Name} Overlaps with Previous in Anim Sequence {Sequence} Time {Time}", *State->GetName(), *Sequence->GetName(), AnimNotifyRef.GetNotify()->GetTime());
+					UE_LOGFMT(LogAnimDatabase, Warning, "MakeFrameRangesFromAnimNotifyState: Anim Notify State {Name} Overlaps with Previous in Anim Sequence {Sequence} Time {Time}", *State->GetName(), *Sequence->GetName(), AnimNotifyRef.GetNotify()->GetTime());
 					AnimNotifyStart = PrevAnimNotifyStopFrame;
 					AnimNotifyStop = FMath::Max(AnimNotifyStop, AnimNotifyStart + 1);
 				}
@@ -1640,7 +1801,7 @@ FAnimDatabaseFrameRanges UAnimDatabaseFrameRangesLibrary::MakeFrameRangesFromAni
 				}
 				else
 				{
-					UE_LOGFMT(LogAnimDatabase, Warning, "Anim Notify State {Name} is less than a single frame in duration in Anim Sequence {Sequence} Time {Time}", *State->GetName(), *Sequence->GetName(), AnimNotifyRef.GetNotify()->GetTime());
+					UE_LOGFMT(LogAnimDatabase, Warning, "MakeFrameRangesFromAnimNotifyState: Anim Notify State {Name} is less than a single frame in duration in Anim Sequence {Sequence} Time {Time}", *State->GetName(), *Sequence->GetName(), AnimNotifyRef.GetNotify()->GetTime());
 				}
 			}
 		}
@@ -1772,6 +1933,11 @@ FAnimDatabaseFrameRanges UAnimDatabaseFrameRangesLibrary::MakeFrameRangesFromAni
 	return MakeFrameRangesFromFrames(UAnimDatabaseFramesLibrary::MakeFramesFromAnimNotify(Database, FrameRanges, Notify));
 }
 
+FAnimDatabaseFrameRanges UAnimDatabaseFrameRangesLibrary::MakeFrameRangesFromSyncMarker(const UAnimDatabase* Database, const FAnimDatabaseFrameRanges& FrameRanges, const FName SyncMarkerName)
+{
+	return MakeFrameRangesFromFrames(UAnimDatabaseFramesLibrary::MakeFramesFromSyncMarker(Database, FrameRanges, SyncMarkerName));
+}
+
 FAnimDatabaseFrames UAnimDatabaseFrameRangesLibrary::MakeFramesAtFrameRangesStarts(const FAnimDatabaseFrameRanges& FrameRanges)
 {
 	if (!FrameRanges.IsValid())
@@ -2662,6 +2828,47 @@ void UAnimDatabaseFrameRangesLibrary::FindCurvesInFrameRanges(TArray<FName>& Out
 	}
 }
 
+void UAnimDatabaseFrameRangesLibrary::FindSyncMarkersInFrameRanges(TArray<FName>& SyncMarkers, const UAnimDatabase* Database, const FAnimDatabaseFrameRanges& FrameRanges)
+{
+	if (!Database)
+	{
+		UE_LOGFMT(LogAnimDatabase, Error, "FindSyncMarkersInFrameRanges: Database is nullptr.");
+		SyncMarkers.Empty();
+		return;
+	}
+
+	if (!FrameRanges.IsValid())
+	{
+		UE_LOGFMT(LogAnimDatabase, Error, "FindSyncMarkersInFrameRanges: Invalid FrameRanges.");
+		SyncMarkers.Empty();
+		return;
+	}
+
+	SyncMarkers.Reset();
+
+	const int32 EntryNum = FrameRanges.FrameRangeSet->GetEntryNum();
+	for (int32 EntryIdx = 0; EntryIdx < EntryNum; EntryIdx++)
+	{
+		const int32 SequenceIdx = FrameRanges.FrameRangeSet->GetEntrySequence(EntryIdx);
+		const UAnimSequence* Sequence = Database->GetAnimSequence(SequenceIdx);
+		if (!Sequence) { continue; }
+
+		const int32 SequenceFrameNum = FMath::Max(Database->GetSequenceFrameNum(SequenceIdx), 1);
+
+		for (const FAnimSyncMarker& SyncMarker : Sequence->AuthoredSyncMarkers)
+		{
+			const int32 SyncMarkerStartStart = FMath::Clamp(
+				FMath::RoundToInt(SyncMarker.Time * Database->GetFrameRate().AsDecimal()),
+				0, SequenceFrameNum - 1);
+
+			if (FrameRanges.FrameRangeSet->Contains(SequenceIdx, SyncMarkerStartStart))
+			{
+				SyncMarkers.AddUnique(SyncMarker.MarkerName);
+			}
+		}
+	}
+}
+
 void UAnimDatabaseFrameRangesLibrary::FindAllAnimNotifyStateClasses(TArray<TSubclassOf<UAnimNotifyState>>& OutClasses, const TSubclassOf<UAnimNotifyState> BaseClass)
 {
 	if (!BaseClass)
@@ -2736,6 +2943,16 @@ FAnimDatabaseFrameRanges UAnimDatabaseFrameRangesFunction_ForceRootLockDisabled:
 	return UAnimDatabaseFrameRangesLibrary::MakeFrameRangesFromForceRootLockDisabled(Database, FrameRanges);
 }
 
+FAnimDatabaseFrameRanges UAnimDatabaseFrameRangesFunction_Looped::MakeFrameRanges_Implementation(const UAnimDatabase* Database, const FAnimDatabaseFrameRanges& FrameRanges) const
+{
+	return UAnimDatabaseFrameRangesLibrary::MakeFrameRangesFromLooped(Database, FrameRanges);
+}
+
+FAnimDatabaseFrameRanges UAnimDatabaseFrameRangesFunction_NotLooped::MakeFrameRanges_Implementation(const UAnimDatabase* Database, const FAnimDatabaseFrameRanges& FrameRanges) const
+{
+	return UAnimDatabaseFrameRangesLibrary::MakeFrameRangesFromNotLooped(Database, FrameRanges);
+}
+
 FAnimDatabaseFrameRanges UAnimDatabaseFrameRangesFunction_AnimNotify::MakeFrameRanges_Implementation(const UAnimDatabase* Database, const FAnimDatabaseFrameRanges& FrameRanges) const
 {
 	return UAnimDatabaseFrameRangesLibrary::MakeFrameRangesFromAnimNotify(Database, FrameRanges, AnimNotify);
diff --git a/Engine/Plugins/Experimental/Animation/AnimDatabase/Source/AnimDatabase/Public/AnimDatabase.h b/Engine/Plugins/Experimental/Animation/AnimDatabase/Source/AnimDatabase/Public/AnimDatabase.h
--- a/Engine/Plugins/Experimental/Animation/AnimDatabase/Source/AnimDatabase/Public/AnimDatabase.h
+++ b/Engine/Plugins/Experimental/Animation/AnimDatabase/Source/AnimDatabase/Public/AnimDatabase.h
@@ -926,6 +926,8 @@ public:
 	static UE_API void EnableRootMotionAndForceRootLockOnSequences(const TArray<UAnimSequence*>& Sequences);
 	static UE_API void EnableRootMotionAndForceRootLockOnSequencesArrayView(TConstArrayView<UAnimSequence*> Sequences);
 
+public:
+
 	/** Removes all AnimNotifies and AnimNotifyStates from the given frame ranges */
 	UFUNCTION(BlueprintCallable, Category = "AnimDatabase", meta = (DevelopmentOnly))
 	static UE_API void DatabaseRemoveNotifies(UAnimDatabase* Database, const FAnimDatabaseFrameRanges& FrameRanges);
@@ -938,6 +940,8 @@ public:
 	UFUNCTION(BlueprintCallable, Category = "AnimDatabase", meta = (DevelopmentOnly, RelativePath, FilePathFilter = "json"))
 	static UE_API void DatabaseImportNotifies(UAnimDatabase* Database, const FAnimDatabaseFrameRanges& FrameRanges, const FFilePath& Path);
 
+public:
+
 	/** Exports the given frame ranges in the database as BVH files into the provided directory */
 	UFUNCTION(BlueprintCallable, Category = "AnimDatabase", meta = (DevelopmentOnly, RelativePath))
 	static UE_API void DatabaseExportAsBVH(UAnimDatabase* Database, const FAnimDatabaseFrameRanges& FrameRanges, const FDirectoryPath& Path, const bool bAddRootAxesChangeRotation = true, const EAnimDatabaseRotationOrder RotationOrder = EAnimDatabaseRotationOrder::XYZ);
@@ -954,40 +958,69 @@ public:
 	UFUNCTION(BlueprintCallable, Category = "AnimDatabase", meta = (DevelopmentOnly, RelativePath, FilePathFilter = "json"))
 	static UE_API void DatabaseExportRootMotion(UAnimDatabase* Database, const FAnimDatabaseFrameRanges& FrameRanges, const FFilePath& Path);
 
+public:
+
 	/** Write the provided float frame attribute back to the animations in the database with the given curve name */
 	UFUNCTION(BlueprintCallable, Category = "AnimDatabase", meta = (DevelopmentOnly))
-	static UE_API void DatabaseAddCurve(UAnimDatabase* Database, const FName CurveName, const FAnimDatabaseFrameAttribute& FloatFrameAttribute, const bool bSparseKeys = false, const bool bShouldTransact = false);
+	static UE_API void DatabaseAddCurve(UAnimDatabase* Database, const FName CurveName, const FAnimDatabaseFrameAttribute& FloatFrameAttribute, const FLinearColor Color = FLinearColor::Red, const bool bSparseKeys = false, const bool bShouldTransact = false);
 
 	/** Write the provided float frame attributes back to the animations in the database with the given curve names */
-	UFUNCTION(BlueprintCallable, Category = "AnimDatabase", meta = (DevelopmentOnly))
-	static UE_API void DatabaseAddCurves(UAnimDatabase* Database, const TArray<FName>& CurveNames, const TArray<FAnimDatabaseFrameAttribute>& FloatFrameAttributes, const bool bSparseKeys = false, const bool bShouldTransact = false);
-	static UE_API void DatabaseAddCurvesFromArrayViews(UAnimDatabase* Database, const TArrayView<const FName> CurveNames, const TArrayView<const FAnimDatabaseFrameAttribute> FloatFrameAttributes, const bool bSparseKeys = false, const bool bShouldTransact = false);
+	UFUNCTION(BlueprintCallable, Category = "AnimDatabase", meta = (DevelopmentOnly, AutoCreateRefTerm = "Colors"))
+	static UE_API void DatabaseAddCurves(UAnimDatabase* Database, const TArray<FName>& CurveNames, const TArray<FAnimDatabaseFrameAttribute>& FloatFrameAttributes, const TArray<FLinearColor>& Colors, const bool bSparseKeys = false, const bool bShouldTransact = false);
+	static UE_API void DatabaseAddCurvesFromArrayViews(UAnimDatabase* Database, const TArrayView<const FName> CurveNames, const TArrayView<const FAnimDatabaseFrameAttribute> FloatFrameAttributes, const TArrayView<const FLinearColor> Colors, const bool bSparseKeys = false, const bool bShouldTransact = false);
 
-	/** Create AnimNotifyStates at the provided FrameRanges in the database using the given track name and class */
-	UFUNCTION(BlueprintCallable, Category = "AnimDatabase", meta = (DevelopmentOnly))
-	static UE_API void DatabaseAddAnimNotifyState(UAnimDatabase* Database, const FName AnimNotifyStateTrackName, const TSubclassOf<UAnimNotifyState> Class, const FAnimDatabaseFrameRanges& FrameRanges);
-
-	/** Create AnimNotifyStates at the provided FrameRanges in the database using the given track names and classes */
-	UFUNCTION(BlueprintCallable, Category = "AnimDatabase", meta = (DevelopmentOnly))
-	static UE_API void DatabaseAddAnimNotifyStates(UAnimDatabase* Database, const TArray<FName>& AnimNotifyStateTrackNames, const TArray<TSubclassOf<UAnimNotifyState>>& Classes, const TArray<FAnimDatabaseFrameRanges>& FrameRanges);
-	static UE_API void DatabaseAddAnimNotifyStatesFromArrayViews(UAnimDatabase* Database, const TArrayView<const FName> AnimNotifyStateTrackNames, const TArrayView<const TSubclassOf<UAnimNotifyState>> Classes, const TArrayView<const FAnimDatabaseFrameRanges> FrameRanges);
+public:
 
 	/** Create AnimNotifies at the provided Frames in the database using the given track name and class */
 	UFUNCTION(BlueprintCallable, Category = "AnimDatabase", meta = (DevelopmentOnly))
-	static UE_API void DatabaseAddAnimNotify(UAnimDatabase* Database, const FName AnimNotifyTrackName, const TSubclassOf<UAnimNotify> Class, const FAnimDatabaseFrames& Frames);
+	static UE_API void DatabaseAddAnimNotify(UAnimDatabase* Database, const FName AnimNotifyTrackName, const TSubclassOf<UAnimNotify> Class, const FAnimDatabaseFrames& Frames, const FLinearColor Color = FLinearColor::Red, const bool bShouldTransact = false);
 
 	/** Create AnimNotifies at the provided Frames in the database using the given track names and classes */
-	UFUNCTION(BlueprintCallable, Category = "AnimDatabase", meta = (DevelopmentOnly))
-	static UE_API void DatabaseAddAnimNotifies(UAnimDatabase* Database, const TArray<FName>& AnimNotifyTrackNames, const TArray<TSubclassOf<UAnimNotify>>& Classes, const TArray<FAnimDatabaseFrames>& Frames);
-	static UE_API void DatabaseAddAnimNotifiesFromArrayViews(UAnimDatabase* Database, const TArrayView<const FName> AnimNotifyTrackNames, const TArrayView<const TSubclassOf<UAnimNotify>> Classes, const TArrayView<const FAnimDatabaseFrames> Frames);
+	UFUNCTION(BlueprintCallable, Category = "AnimDatabase", meta = (DevelopmentOnly, AutoCreateRefTerm = "Colors"))
+	static UE_API void DatabaseAddAnimNotifies(UAnimDatabase* Database, const TArray<FName>& AnimNotifyTrackNames, const TArray<TSubclassOf<UAnimNotify>>& Classes, const TArray<FAnimDatabaseFrames>& Frames, const TArray<FLinearColor>& Colors, const bool bShouldTransact = false);
+	static UE_API void DatabaseAddAnimNotifiesFromArrayViews(UAnimDatabase* Database, const TArrayView<const FName> AnimNotifyTrackNames, const TArrayView<const TSubclassOf<UAnimNotify>> Classes, const TArrayView<const FAnimDatabaseFrames> Frames, const TArrayView<const FLinearColor> Colors, const bool bShouldTransact = false);
 
 	/** Create AnimNotifies at the provided Frames from a duplicate of the given AnimNotify */
 	UFUNCTION(BlueprintCallable, Category = "AnimDatabase", meta = (DevelopmentOnly))
-	static UE_API void DatabaseAddAnimNotifyObject(UAnimDatabase* Database, const FAnimDatabaseFrames& Frames, const FName& TrackName, const UAnimNotify* AnimNotify);
+	static UE_API void DatabaseAddAnimNotifyObject(UAnimDatabase* Database, const FName& TrackName, UAnimNotify* AnimNotify, const FAnimDatabaseFrames& Frames, const FLinearColor Color = FLinearColor::Red, const bool bShouldTransact = false);
+
+	/** Create AnimNotifies at the provided Frames in the database using the given track names and duplicates of the given AnimNotifies */
+	UFUNCTION(BlueprintCallable, Category = "AnimDatabase", meta = (DevelopmentOnly, AutoCreateRefTerm="Colors"))
+	static UE_API void DatabaseAddAnimNotifyObjects(UAnimDatabase* Database, const TArray<FName>& AnimNotifyTrackNames, const TArray<UAnimNotify*>& AnimNotifies, const TArray<FAnimDatabaseFrames>& Frames, const TArray<FLinearColor>& Colors, const bool bShouldTransact = false);
+	static UE_API void DatabaseAddAnimNotifyObjectsFromArrayViews(UAnimDatabase* Database, const TArrayView<const FName> AnimNotifyTrackNames, const TArrayView<UAnimNotify* const> AnimNotifies, const TArrayView<const FAnimDatabaseFrames> Frames, const TArrayView<const FLinearColor> Colors, const bool bShouldTransact = false);
+
+public:
+
+	/** Create AnimNotifyStates at the provided FrameRanges in the database using the given track name and class */
+	UFUNCTION(BlueprintCallable, Category = "AnimDatabase", meta = (DevelopmentOnly))
+	static UE_API void DatabaseAddAnimNotifyState(UAnimDatabase* Database, const FName AnimNotifyStateTrackName, const TSubclassOf<UAnimNotifyState> Class, const FAnimDatabaseFrameRanges& FrameRanges, const FLinearColor Color = FLinearColor::Red, const bool bShouldTransact = false);
+
+	/** Create AnimNotifyStates at the provided FrameRanges in the database using the given track names and classes */
+	UFUNCTION(BlueprintCallable, Category = "AnimDatabase", meta = (DevelopmentOnly, AutoCreateRefTerm = "Colors"))
+	static UE_API void DatabaseAddAnimNotifyStates(UAnimDatabase* Database, const TArray<FName>& AnimNotifyStateTrackNames, const TArray<TSubclassOf<UAnimNotifyState>>& Classes, const TArray<FAnimDatabaseFrameRanges>& FrameRanges, const TArray<FLinearColor>& Colors, const bool bShouldTransact = false);
+	static UE_API void DatabaseAddAnimNotifyStatesFromArrayViews(UAnimDatabase* Database, const TArrayView<const FName> AnimNotifyStateTrackNames, const TArrayView<const TSubclassOf<UAnimNotifyState>> Classes, const TArrayView<const FAnimDatabaseFrameRanges> FrameRanges, const TArrayView<const FLinearColor> Colors, const bool bShouldTransact = false);
 
 	/** Create AnimNotifyStates at the provided FrameRanges from a duplicate of the given AnimNotifyState */
 	UFUNCTION(BlueprintCallable, Category = "AnimDatabase", meta = (DevelopmentOnly))
-	static UE_API void DatabaseAddAnimNotifyStateObject(UAnimDatabase* Database, const FAnimDatabaseFrameRanges& FrameRanges, const FName& TrackName, const UAnimNotifyState* AnimNotifyState);
+	static UE_API void DatabaseAddAnimNotifyStateObject(UAnimDatabase* Database, const FName TrackName, UAnimNotifyState* AnimNotifyState, const FAnimDatabaseFrameRanges& FrameRanges, const FLinearColor Color = FLinearColor::Red, const bool bShouldTransact = false);
+
+	/** Create AnimNotifyStates at the provided FrameRanges in the database using the given track names and duplicates of the given AnimNotifyStates */
+	UFUNCTION(BlueprintCallable, Category = "AnimDatabase", meta = (DevelopmentOnly, AutoCreateRefTerm = "Colors"))
+	static UE_API void DatabaseAddAnimNotifyStateObjects(UAnimDatabase* Database, const TArray<FName>& AnimNotifyStateTrackNames, const TArray<UAnimNotifyState*>& AnimNotifyStates, const TArray<FAnimDatabaseFrameRanges>& FrameRanges, const TArray<FLinearColor>& Colors, const bool bShouldTransact = false);
+	static UE_API void DatabaseAddAnimNotifyStateObjectsFromArrayViews(UAnimDatabase* Database, const TArrayView<const FName> AnimNotifyStateTrackNames, const TArrayView<UAnimNotifyState* const> AnimNotifyStates, const TArrayView<const FAnimDatabaseFrameRanges> FrameRanges, const TArrayView<const FLinearColor> Colors, const bool bShouldTransact = false);
+
+public:
+
+	/** Create SyncMarkers at the provided Frames in the database using the given track name and sync marker name */
+	UFUNCTION(BlueprintCallable, Category = "AnimDatabase", meta = (DevelopmentOnly))
+	static UE_API void DatabaseAddSyncMarker(UAnimDatabase* Database, const FName SyncMarkerTrackName, const FName SyncMarkerName, const FAnimDatabaseFrames& Frames, const FLinearColor Color = FLinearColor::Red, const bool bShouldTransact = false);
+
+	/** Create SyncMarkers at the provided Frames in the database using the given track names and sync marker names */
+	UFUNCTION(BlueprintCallable, Category = "AnimDatabase", meta = (DevelopmentOnly, AutoCreateRefTerm = "Colors"))
+	static UE_API void DatabaseAddSyncMarkers(UAnimDatabase* Database, const TArray<FName>& SyncMarkerTrackNames, const TArray<FName>& SyncMarkerNames, const TArray<FAnimDatabaseFrames>& Frames, const TArray<FLinearColor>& Colors, const bool bShouldTransact = false);
+	static UE_API void DatabaseAddSyncMarkersFromArrayViews(UAnimDatabase* Database, const TArrayView<const FName> SyncMarkerTrackNames, const TArrayView<const  FName> SyncMarkerNames, const TArrayView<const FAnimDatabaseFrames> Frames, const TArrayView<const FLinearColor> Colors, const bool bShouldTransact = false);
+
+public:
 
 	/** Set the given bone's local transform in the database at the given transform frame attribute */
 	UFUNCTION(BlueprintCallable, Category = "AnimDatabase", meta = (DevelopmentOnly))
@@ -1016,6 +1049,8 @@ public:
 	static UE_API void DatabaseResetBoneTransforms(UAnimDatabase* Database, const FAnimDatabaseFrameRanges& FrameRanges, const TArray<FName>& BoneNames, const bool bShouldTransact = false);
 	static UE_API void DatabaseResetBoneTransformsFromArrayViews(UAnimDatabase* Database, const FAnimDatabaseFrameRanges& FrameRanges, const TArrayView<const FName> BoneNames, const bool bShouldTransact = false);
 
+public:
+
 	/** Strips unused bone tracks from the data to reduce the file size on disk */
 	UFUNCTION(BlueprintCallable, Category = "AnimDatabase", meta = (DevelopmentOnly))
 	static UE_API void DatabaseStripUnusedBoneTracks(UAnimDatabase* Database, const FAnimDatabaseFrameRanges& FrameRanges, const bool bShouldTransact = false);
@@ -1564,6 +1599,9 @@ public:
 	UPROPERTY(EditAnywhere, Instanced, Category = "Settings")
 	TObjectPtr<UAnimNotifyState> AnimNotifyState;
 
+	UPROPERTY(EditAnywhere, Category = "Settings")
+	FLinearColor Color = FLinearColor::Red;
+
 #endif
 
 #if WITH_EDITOR
@@ -1618,6 +1656,9 @@ public:
 	UPROPERTY(EditAnywhere, Instanced, Category = "Settings")
 	TObjectPtr<UAnimDatabaseFramesFunction> Frames;
 
+	UPROPERTY(EditAnywhere, Category = "Settings")
+	FLinearColor Color = FLinearColor::Red;
+
 #endif
 
 #if WITH_EDITOR
@@ -1648,6 +1689,9 @@ public:
 	UPROPERTY(EditAnywhere, Instanced, Category = "Settings")
 	TObjectPtr<UAnimDatabaseFrameRangesFunction> FrameRanges;
 
+	UPROPERTY(EditAnywhere, Category = "Settings")
+	FLinearColor Color = FLinearColor::Red;
+
 #endif
 
 #if WITH_EDITOR
@@ -1675,6 +1719,9 @@ public:
 	UPROPERTY(EditAnywhere, Instanced, Category = "Settings")
 	TObjectPtr<UAnimDatabaseFrameAttributeFunction> FrameAttribute;
 
+	UPROPERTY(EditAnywhere, Category = "Settings")
+	FLinearColor Color = FLinearColor::Red;
+
 	UPROPERTY(EditAnywhere, Category = "Settings")
 	bool bSparseKeys = true;
 
@@ -1689,9 +1736,9 @@ public:
 #endif
 };
 
-/** Function for adding contact curves  */
-UCLASS(BlueprintType, Blueprintable, meta = (DisplayName = "Add Contact Curves"))
-class UAnimDatabaseFunction_AddContactCurves : public UAnimDatabaseFunction
+/** Function for processing contacts */
+UCLASS(BlueprintType, Blueprintable, DontCollapseCategories, meta = (DisplayName = "Process Contacts"))
+class UAnimDatabaseFunction_ProcessContacts : public UAnimDatabaseFunction
 {
 	GENERATED_BODY()
 
@@ -1699,44 +1746,165 @@ class UAnimDatabaseFunction_AddContactCurves : public UAnimDatabaseFunction
 
 public:
 
+	/** If to add foot speed curves */
+	UPROPERTY(EditAnywhere, Category = "Foot Speed Curves")
+	bool bAddFootSpeedCurves = false;
+
+	/** Name of the foot speed curve for the left foot */
+	UPROPERTY(EditAnywhere, Category = "Foot Speed Curves", meta = (EditCondition = "bAddFootSpeedCurves", HideEditConditionToggle))
+	FName LeftFootSpeedCurveName = TEXT("footspeed_l");
+
+	/** Name of the foot speed curve for the right foot */
+	UPROPERTY(EditAnywhere, Category = "Foot Speed Curves", meta = (EditCondition = "bAddFootSpeedCurves", HideEditConditionToggle))
+	FName RightFootSpeedCurveName = TEXT("footspeed_r");
+
+	/** Name of the left bone to generate the foot speed curve for */
+	UPROPERTY(EditAnywhere, Category = "Foot Speed Curves", meta = (EditCondition = "bAddFootSpeedCurves", HideEditConditionToggle))
+	FName LeftFootSpeedBoneName = TEXT("ball_l");
+
+	/** Name of the right bone to generate the foot speed curve for */
+	UPROPERTY(EditAnywhere, Category = "Foot Speed Curves", meta = (EditCondition = "bAddFootSpeedCurves", HideEditConditionToggle))
+	FName RightFootSpeedBoneName = TEXT("ball_r");
+
+	/** Color of the left foot speed curve */
+	UPROPERTY(EditAnywhere, Category = "Foot Speed Curves", meta = (EditCondition = "bAddFootSpeedCurves", HideEditConditionToggle))
+	FLinearColor LeftFootSpeedCurveColor = FLinearColor::Red;
+
+	/** Color of the right foot speed curve */
+	UPROPERTY(EditAnywhere, Category = "Foot Speed Curves", meta = (EditCondition = "bAddFootSpeedCurves", HideEditConditionToggle))
+	FLinearColor RightFootSpeedCurveColor = FLinearColor::Blue;
+
+public:
+
+	/** If to add contact curves */
+	UPROPERTY(EditAnywhere, Category = "Contact Curves")
+	bool bAddContactCurves = false;
+
 	/** Name of the contact curve for the left foot */
-	UPROPERTY(EditAnywhere, Category = "Settings")
+	UPROPERTY(EditAnywhere, Category = "Contact Curves", meta = (EditCondition = "bAddContactCurves", HideEditConditionToggle))
 	FName LeftContactCurveName = TEXT("contact_l");
 
 	/** Name of the contact curve for the right foot */
-	UPROPERTY(EditAnywhere, Category = "Settings")
+	UPROPERTY(EditAnywhere, Category = "Contact Curves", meta = (EditCondition = "bAddContactCurves", HideEditConditionToggle))
 	FName RightContactCurveName = TEXT("contact_r");
 
 	/** Name of the left bone to generate the contact curve for */
-	UPROPERTY(EditAnywhere, Category = "Settings")
+	UPROPERTY(EditAnywhere, Category = "Contact Curves", meta = (EditCondition = "bAddContactCurves", HideEditConditionToggle))
 	FName LeftContactBoneName = TEXT("ball_l");
 
 	/** Name of the right bone to generate the contact curve for */
-	UPROPERTY(EditAnywhere, Category = "Settings")
+	UPROPERTY(EditAnywhere, Category = "Contact Curves", meta = (EditCondition = "bAddContactCurves", HideEditConditionToggle))
 	FName RightContactBoneName = TEXT("ball_r");
 
+	/** Color of the left contact curve */
+	UPROPERTY(EditAnywhere, Category = "Contact Curves", meta = (EditCondition = "bAddContactCurves", HideEditConditionToggle))
+	FLinearColor LeftContactCurveColor = FLinearColor::Red;
+
+	/** Color of the right contact curve */
+	UPROPERTY(EditAnywhere, Category = "Contact Curves", meta = (EditCondition = "bAddContactCurves", HideEditConditionToggle))
+	FLinearColor RightContactCurveColor = FLinearColor::Blue;
+
+public:
+
+	/** If to add anim notifies */
+	UPROPERTY(EditAnywhere, Category = "Anim Notifies")
+	bool bAddAnimNotifies = false;
+
+	/** The track name for the left anim notifies */
+	UPROPERTY(EditAnywhere, Category = "Anim Notifies", meta = (EditCondition = "bAddAnimNotifies", HideEditConditionToggle))
+	FName LeftContactAnimNotifyTrackName = TEXT("Footstep Left");
+
+	/** The track name for the right anim notifies */
+	UPROPERTY(EditAnywhere, Category = "Anim Notifies", meta = (EditCondition = "bAddAnimNotifies", HideEditConditionToggle))
+	FName RightContactAnimNotifyTrackName = TEXT("Footstep Right");
+
+	/** The object to use for the left contact anim notify */
+	UPROPERTY(EditAnywhere, Instanced, Category = "Anim Notifies", meta = (EditCondition = "bAddAnimNotifies", HideEditConditionToggle))
+	TObjectPtr<UAnimNotify> LeftContactAnimNotify;
+
+	/** The object to use for the right contact anim notify */
+	UPROPERTY(EditAnywhere, Instanced, Category = "Anim Notifies", meta = (EditCondition = "bAddAnimNotifies", HideEditConditionToggle))
+	TObjectPtr<UAnimNotify> RightContactAnimNotify;
+
+	/** Color of the left anim notify track */
+	UPROPERTY(EditAnywhere, Category = "Anim Notifies", meta = (EditCondition = "bAddAnimNotifies", HideEditConditionToggle))
+	FLinearColor LeftContactAnimNotifyColor = FLinearColor::Red;
+
+	/** Color of the right anim notify track */
+	UPROPERTY(EditAnywhere, Category = "Anim Notifies", meta = (EditCondition = "bAddAnimNotifies", HideEditConditionToggle))
+	FLinearColor RightContactAnimNotifyColor = FLinearColor::Blue;
+
+public:
+
+	/** If to add sync markers */
+	UPROPERTY(EditAnywhere, Category = "Sync Markers")
+	bool bAddSyncMarkers = false;
+
+	/** The track name for the left sync markers */
+	UPROPERTY(EditAnywhere, Category = "Sync Markers", meta = (EditCondition = "bAddSyncMarkers", HideEditConditionToggle))
+	FName LeftContactSyncMarkerTrackName = TEXT("FootSyncMarkers");
+
+	/** The track name for the right sync markers */
+	UPROPERTY(EditAnywhere, Category = "Sync Markers", meta = (EditCondition = "bAddSyncMarkers", HideEditConditionToggle))
+	FName RightContactSyncMarkerTrackName = TEXT("FootSyncMarkers");
+
+	/** The name of the left sync marker */
+	UPROPERTY(EditAnywhere, Category = "Sync Markers", meta = (EditCondition = "bAddSyncMarkers", HideEditConditionToggle))
+	FName LeftContactSyncMarker = TEXT("L");
+
+	/** The name of the right sync marker */
+	UPROPERTY(EditAnywhere, Category = "Sync Markers", meta = (EditCondition = "bAddSyncMarkers", HideEditConditionToggle))
+	FName RightContactSyncMarker = TEXT("R");
+
+	/** Color of the left sync marker track */
+	UPROPERTY(EditAnywhere, Category = "Sync Markers", meta = (EditCondition = "bAddSyncMarkers", HideEditConditionToggle))
+	FLinearColor LeftContactSyncMarkerColor = FLinearColor(0.5f, 0.0, 0.5f, 1.0f);
+
+	/** Color of the right sync marker track */
+	UPROPERTY(EditAnywhere, Category = "Sync Markers", meta = (EditCondition = "bAddSyncMarkers", HideEditConditionToggle))
+	FLinearColor RightContactSyncMarkerColor = FLinearColor(0.5f, 0.0, 0.5f, 1.0f);
+
+public:
+
 	/** Maximum height (in cm) for which the bone can be considered in contact */
-	UPROPERTY(EditAnywhere, Category = "Settings")
+	UPROPERTY(EditAnywhere, Category = "Filtering", meta = (EditCondition = "bAddContactCurves || bAddAnimNotifies || bAddSyncMarkers", HideEditConditionToggle))
 	float HeightThreshold = 20.0f;
 
 	/** Maximum velocity (in cm/s) for which the bone can be considered in contact */
-	UPROPERTY(EditAnywhere, Category = "Settings")
+	UPROPERTY(EditAnywhere, Category = "Filtering", meta = (EditCondition = "bAddContactCurves || bAddAnimNotifies || bAddSyncMarkers", HideEditConditionToggle))
 	float VelocityThreshold = 25.0f;
 
+	/** If to filter the contact curves */
+	UPROPERTY(EditAnywhere, Category = "Filtering", meta = (EditCondition = "bAddContactCurves || bAddAnimNotifies || bAddSyncMarkers", HideEditConditionToggle))
+	bool bFilterContactCurves = true;
+
 	/**
 	 * Filter width (in frames) for the majority vote filter used to remove very short contact transitions.
 	 * Contact periods (or non-contact periods) less than half this duration will be removed.
 	 */
-	UPROPERTY(EditAnywhere, Category = "Settings")
+	UPROPERTY(EditAnywhere, Category = "Filtering", meta = (EditCondition = "(bAddContactCurves || bAddAnimNotifies || bAddSyncMarkers) && bFilterContactCurves", HideEditConditionToggle))
 	int32 MajorityVoteFilterWidth = 5;
 
-	/* Amount of smoothing to apply to the resulting contact curve (in frames) */
-	UPROPERTY(EditAnywhere, Category = "Settings")
+	/** If to smooth the contact curves */
+	UPROPERTY(EditAnywhere, Category = "Filtering", meta = (EditCondition = "bAddContactCurves || bAddAnimNotifies || bAddSyncMarkers", HideEditConditionToggle))
+	bool bSmoothContactCurves = true;
+
+	/** Amount of smoothing to apply to the resulting contact curve (in frames) */
+	UPROPERTY(EditAnywhere, Category = "Filtering", meta = (EditCondition = "(bAddContactCurves || bAddAnimNotifies || bAddSyncMarkers) && bSmoothContactCurves", HideEditConditionToggle))
 	float SmoothingAmount = 3.0f;
 
-	UPROPERTY(EditAnywhere, Category = "Settings")
+	/** If to write curves using sparse keys */
+	UPROPERTY(EditAnywhere, Category = "Filtering", meta = (EditCondition = "bAddFootSpeedCurves || bAddContactCurves", HideEditConditionToggle))
 	bool bSparseKeys = true;
 
+	/** At what threshold of the contact curve to use to detect Anim Notifies and Sync Markers */
+	UPROPERTY(EditAnywhere, Category = "Filtering", meta = (EditCondition = "bAddAnimNotifies || bAddSyncMarkers", HideEditConditionToggle))
+	float EventThreshold = 0.5f;
+
+	/** The amount of time before the contact (in seconds) to place the location of Anim Notifies and Sync Markers */
+	UPROPERTY(EditAnywhere, Category = "Filtering", meta = (EditCondition = "bAddAnimNotifies || bAddSyncMarkers", HideEditConditionToggle, ForceUnits="s"))
+	float EventOffset = 0.0f;
+
 #endif
 
 #if WITH_EDITOR
diff --git a/Engine/Plugins/Experimental/Animation/AnimDatabase/Source/AnimDatabase/Public/AnimDatabaseFrameAttribute.h b/Engine/Plugins/Experimental/Animation/AnimDatabase/Source/AnimDatabase/Public/AnimDatabaseFrameAttribute.h
--- a/Engine/Plugins/Experimental/Animation/AnimDatabase/Source/AnimDatabase/Public/AnimDatabaseFrameAttribute.h
+++ b/Engine/Plugins/Experimental/Animation/AnimDatabase/Source/AnimDatabase/Public/AnimDatabaseFrameAttribute.h
@@ -1867,8 +1867,10 @@ public:
 	 * @param BoneName							Name of the bone to compute the contact state for
 	 * @param HeightThreshold					Maximum height (in cm) for which the bone can be considered in contact
 	 * @param VelocityThreshold					Maximum velocity (in cm/s) for which the bone can be considered in contact
+	 * @param bFilterCurve						If to apply the majority vote filter to the curve
 	 * @param MajorityVoteFilterWidth			Filter width (in frames) for the majority vote filter used to remove very short contact transitions.
 	 *											contact periods (or non-contact periods) less than half this duration will be removed.
+	 * @param bSmoothCurve						If to apply the smoothing filter to the curve
 	 * @param SmoothingAmount					Amount of smoothing to apply to the resulting contact curve (in frames)
 	 */
 	UFUNCTION(BlueprintPure, Category = "AnimDatabase", meta = (NotBlueprintThreadSafe))
@@ -1878,7 +1880,9 @@ public:
 		const FName BoneName = TEXT("ball_l"),
 		const float HeightThreshold = 20.0f,
 		const float VelocityThreshold = 25.0f,
+		const bool bFilterCurve = true,
 		const int32 MajorityVoteFilterWidth = 5,
+		const bool bSmoothCurve = true,
 		const float SmoothingAmount = 3.0f);
 
 	/**
@@ -1889,6 +1893,7 @@ public:
 	 * @param BoneName							Name of the bone to compute the contact state for
 	 * @param HeightThreshold					Maximum height (in cm) for which the bone can be considered in contact
 	 * @param VelocityThreshold					Maximum velocity (in cm/s) for which the bone can be considered in contact
+	 * @param bFilterCurve						If to apply the majority vote filter to the curve
 	 * @param MajorityVoteFilterWidth			Filter width (in frames) for the majority vote filter used to remove very short contact transitions.
 	 *											contact periods (or non-contact periods) less than half this duration will be removed.
 	 */
@@ -1899,6 +1904,7 @@ public:
 		const FName BoneName = TEXT("ball_l"),
 		const float HeightThreshold = 20.0f,
 		const float VelocityThreshold = 25.0f,
+		const bool bFilterCurve = true,
 		const int32 MajorityVoteFilterWidth = 5);
 
 	/**
@@ -1909,8 +1915,10 @@ public:
 	 * @param BoneName							Name of the bone to compute the contact state for
 	 * @param HeightThreshold					Maximum height (in cm) for which the bone can be considered in contact
 	 * @param VelocityThreshold					Maximum velocity (in cm/s) for which the bone can be considered in contact
+	 * @param bFilterCurve						If to apply the majority vote filter to the curve
 	 * @param MajorityVoteFilterWidth			Filter width (in frames) for the majority vote filter used to remove very short contact transitions.
 	 *											contact periods (or non-contact periods) less than half this duration will be removed.
+	 * @param EventOffset						The amount of time before the contact (in seconds) to have the event
 	 */
 	UFUNCTION(BlueprintPure, Category = "AnimDatabase", meta = (NotBlueprintThreadSafe))
 	static UE_API FAnimDatabaseFrameAttribute FrameAttributeContactEvent(
@@ -1919,7 +1927,9 @@ public:
 		const FName BoneName = TEXT("ball_l"),
 		const float HeightThreshold = 20.0f,
 		const float VelocityThreshold = 25.0f,
-		const int32 MajorityVoteFilterWidth = 5);
+		const bool bFilterCurve = true,
+		const int32 MajorityVoteFilterWidth = 5,
+		const float EventOffset = 0.0f);
 
 	/**
 	 * Fits a transform from the given location and direction attributes
@@ -2064,15 +2074,23 @@ public:
 	UPROPERTY(EditAnywhere, Category = "Settings")
 	float VelocityThreshold = 25.0f;
 
+	/** If to filter the curve */
+	UPROPERTY(EditAnywhere, Category = "Settings")
+	bool bFilterCurve = true;
+
 	/**
 	 * Filter width (in frames) for the majority vote filter used to remove very short contact transitions.
 	 *Contact periods (or non-contact periods) less than half this duration will be removed.
 	 */
-	UPROPERTY(EditAnywhere, Category = "Settings")
+	UPROPERTY(EditAnywhere, Category = "Settings", meta = (EditCondition = "bFilterCurve", HideEditConditionToggle))
 	int32 MajorityVoteFilterWidth = 5;
 
-	/* Amount of smoothing to apply to the resulting contact curve (in frames) */
+	/** If to smooth the curve */
 	UPROPERTY(EditAnywhere, Category = "Settings")
+	bool bSmoothCurve = true;
+
+	/* Amount of smoothing to apply to the resulting contact curve (in frames) */
+	UPROPERTY(EditAnywhere, Category = "Settings", meta = (EditCondition = "bSmoothCurve", HideEditConditionToggle))
 	float SmoothingAmount = 3.0f;
 
 public:
@@ -2080,6 +2098,45 @@ public:
 	UE_API virtual FAnimDatabaseFrameAttribute MakeFrameAttribute_Implementation(const UAnimDatabase* Database, const FAnimDatabaseFrameRanges& FrameRanges) const override;
 };
 
+/** Computes a contact event attribute */
+UCLASS(BlueprintType, Blueprintable, meta = (DisplayName = "Contact Event"))
+class UAnimDatabaseFrameAttributeFunction_ContactEvent : public UAnimDatabaseFrameAttributeFunction
+{
+	GENERATED_BODY()
+
+public:
+
+	/** Name of the bone to compute the contact state for */
+	UPROPERTY(EditAnywhere, Category = "Settings")
+	FName BoneName = NAME_None;
+
+	/** Maximum height (in cm) for which the bone can be considered in contact */
+	UPROPERTY(EditAnywhere, Category = "Settings")
+	float HeightThreshold = 20.0f;
+
+	/** Maximum velocity (in cm/s) for which the bone can be considered in contact */
+	UPROPERTY(EditAnywhere, Category = "Settings")
+	float VelocityThreshold = 25.0f;
+
+	/** If to filter the curve */
+	UPROPERTY(EditAnywhere, Category = "Settings")
+	bool bFilterCurve = true;
+
+	/**
+	 * Filter width (in frames) for the majority vote filter used to remove very short contact transitions.
+	 *Contact periods (or non-contact periods) less than half this duration will be removed.
+	 */
+	UPROPERTY(EditAnywhere, Category = "Settings", meta = (EditCondition = "bFilterCurve", HideEditConditionToggle))
+	int32 MajorityVoteFilterWidth = 5;
+
+	/** The amount of time before the contact (in seconds) to have the event */
+	UPROPERTY(EditAnywhere, Category = "Filtering", meta = (ForceUnits = "s"))
+	float EventOffset = 0.0f;
+
+public:
+
+	UE_API virtual FAnimDatabaseFrameAttribute MakeFrameAttribute_Implementation(const UAnimDatabase* Database, const FAnimDatabaseFrameRanges& FrameRanges) const override;
+};
 
 /** Fits a transform to the given locations and directions */
 UCLASS(BlueprintType, Blueprintable, meta = (DisplayName = "Fit Transform"))
diff --git a/Engine/Plugins/Experimental/Animation/AnimDatabase/Source/AnimDatabase/Public/AnimDatabaseFrameRanges.h b/Engine/Plugins/Experimental/Animation/AnimDatabase/Source/AnimDatabase/Public/AnimDatabaseFrameRanges.h
--- a/Engine/Plugins/Experimental/Animation/AnimDatabase/Source/AnimDatabase/Public/AnimDatabaseFrameRanges.h
+++ b/Engine/Plugins/Experimental/Animation/AnimDatabase/Source/AnimDatabase/Public/AnimDatabaseFrameRanges.h
@@ -98,6 +98,14 @@ struct TStructOpsTypeTraits<FAnimDatabaseFrameRanges> : TStructOpsTypeTraitsBase
 	};
 };
 
+UENUM(BlueprintType)
+enum class EAnimDatabaseFrameShiftBehavior : uint8
+{
+	// Remove frames which are shifted to beyond the start or end of an anim sequence
+	Remove = 0,
+	// Clamp frames which are shifted to beyond the start or end of an anim sequence to the start or end
+	Clamp = 1,
+};
 
 UCLASS(BlueprintType, meta = (BlueprintThreadSafe))
 class UAnimDatabaseFramesLibrary : public UBlueprintFunctionLibrary
@@ -122,6 +130,10 @@ public:
 	UFUNCTION(BlueprintPure, Category = "AnimDatabase", meta = (NotBlueprintThreadSafe))
 	static UE_API FAnimDatabaseFrames MakeFramesFromAnimNotify(const UAnimDatabase* Database, const FAnimDatabaseFrameRanges& FrameRanges, const TSubclassOf<UAnimNotify> Notify);
 
+	/** Make a set of Frames from the times given by all the instances of the given SyncMarker in the database */
+	UFUNCTION(BlueprintPure, Category = "AnimDatabase", meta = (NotBlueprintThreadSafe))
+	static UE_API FAnimDatabaseFrames MakeFramesFromSyncMarker(const UAnimDatabase* Database, const FAnimDatabaseFrameRanges& FrameRanges, const FName SyncMarkerName);
+
 	/** Make a set of Frames from the start times of all the instances of the given UAnimNotifyState class in the database */
 	UFUNCTION(BlueprintPure, Category = "AnimDatabase", meta = (NotBlueprintThreadSafe))
 	static UE_API FAnimDatabaseFrames MakeFramesFromAnimNotifyStateStart(const UAnimDatabase* Database, const FAnimDatabaseFrameRanges& FrameRanges, const TSubclassOf<UAnimNotifyState> State);
@@ -149,13 +161,13 @@ public:
 	UFUNCTION(BlueprintPure, Category = "AnimDatabase")
 	static UE_API bool FramesIsEmpty(const FAnimDatabaseFrames& Frames);
 
-	/** Make a set of Frames at a given number of frames before another set of frames. Will discard frames that fall out of the valid sequence range. */
+	/** Make a set of Frames at a given number of frames before another set of frames. Will use ShiftBehavior to decided if to clamp or discard frames that fall out of the valid sequence range. */
 	UFUNCTION(BlueprintPure, Category = "AnimDatabase", meta = (NotBlueprintThreadSafe, CompactNodeTitle = "BEFORE"))
-	static UE_API FAnimDatabaseFrames FramesBefore(const UAnimDatabase* Database, const FAnimDatabaseFrames& Frames, const int32 FramesBefore);
+	static UE_API FAnimDatabaseFrames FramesBefore(const UAnimDatabase* Database, const FAnimDatabaseFrames& Frames, const int32 FramesBefore, const EAnimDatabaseFrameShiftBehavior ShiftBehavior = EAnimDatabaseFrameShiftBehavior::Remove);
 
-	/** Make a set of Frames at a given number of frames after another set of frames. Will discard frames that fall out of the valid sequence range. */
+	/** Make a set of Frames at a given number of frames after another set of frames. Will use ShiftBehavior to decided if to clamp or discard frames that fall out of the valid sequence range. */
 	UFUNCTION(BlueprintPure, Category = "AnimDatabase", meta = (NotBlueprintThreadSafe, CompactNodeTitle = "AFTER"))
-	static UE_API FAnimDatabaseFrames FramesAfter(const UAnimDatabase* Database, const FAnimDatabaseFrames& Frames, const int32 FramesAfter);
+	static UE_API FAnimDatabaseFrames FramesAfter(const UAnimDatabase* Database, const FAnimDatabaseFrames& Frames, const int32 FramesAfter, const EAnimDatabaseFrameShiftBehavior ShiftBehavior = EAnimDatabaseFrameShiftBehavior::Remove);
 
 	/** Checks if two sets of Frames are equal */
 	UFUNCTION(BlueprintPure, Category = "AnimDatabase", meta = (CompactNodeTitle = "EQUAL"))
@@ -356,6 +368,14 @@ public:
 	UFUNCTION(BlueprintPure, Category = "AnimDatabase")
 	static UE_API FAnimDatabaseFrameRanges MakeFrameRangesFromForceRootLockDisabled(const UAnimDatabase* Database, const FAnimDatabaseFrameRanges& FrameRanges);
 
+	/** Make a set of FrameRanges for just the animations which are not looped in a set of frame ranges */
+	UFUNCTION(BlueprintPure, Category = "AnimDatabase")
+	static UE_API FAnimDatabaseFrameRanges MakeFrameRangesFromNotLooped(const UAnimDatabase* Database, const FAnimDatabaseFrameRanges& FrameRanges);
+
+	/** Make a set of FrameRanges for just the animations which are looped in a set of frame ranges */
+	UFUNCTION(BlueprintPure, Category = "AnimDatabase")
+	static UE_API FAnimDatabaseFrameRanges MakeFrameRangesFromLooped(const UAnimDatabase* Database, const FAnimDatabaseFrameRanges& FrameRanges);
+
 	/** Make a set of FrameRanges for a single sequence in a database */
 	UFUNCTION(BlueprintPure, Category = "AnimDatabase")
 	static UE_API FAnimDatabaseFrameRanges MakeFrameRangesFromSequenceIndex(const UAnimDatabase* Database, const int32 SequenceIdx);
@@ -482,6 +502,10 @@ public:
 	UFUNCTION(BlueprintPure, Category = "AnimDatabase")
 	static UE_API FAnimDatabaseFrameRanges MakeFrameRangesFromAnimNotify(const UAnimDatabase* Database, const FAnimDatabaseFrameRanges& FrameRanges, const TSubclassOf<UAnimNotify> Notify);
 
+	/** Make a set of single frame FrameRanges from the given sync marker */
+	UFUNCTION(BlueprintPure, Category = "AnimDatabase")
+	static UE_API FAnimDatabaseFrameRanges MakeFrameRangesFromSyncMarker(const UAnimDatabase* Database, const FAnimDatabaseFrameRanges& FrameRanges, const FName SyncMarkerName);
+
 
 public:
 
@@ -658,6 +682,10 @@ public:
 	UFUNCTION(BlueprintPure, Category = "AnimDatabase")
 	static UE_API void FindCurvesInFrameRanges(TArray<FName>& OutCurves, const UAnimDatabase* Database, const FAnimDatabaseFrameRanges& FrameRanges);
 
+	/** Finds all of the sync markers which are present in the given frame ranges */
+	UFUNCTION(BlueprintPure, Category = "AnimDatabase")
+	static UE_API void FindSyncMarkersInFrameRanges(TArray<FName>& SyncMarkers, const UAnimDatabase* Database, const FAnimDatabaseFrameRanges& FrameRanges);
+
 	/** Finds all of the AnimNotifyState Classes which are subclasses of the given base class (exclusive). Only returns classes loaded in memory. */
 	UFUNCTION(BlueprintPure, Category = "AnimDatabase", meta=(BaseClass = "/Script/Engine.AnimNotifyState"))
 	static UE_API void FindAllAnimNotifyStateClasses(TArray<TSubclassOf<UAnimNotifyState>>& OutClasses, const TSubclassOf<UAnimNotifyState> BaseClass);
@@ -789,6 +817,28 @@ public:
 	UE_API virtual FAnimDatabaseFrameRanges MakeFrameRanges_Implementation(const UAnimDatabase* Database, const FAnimDatabaseFrameRanges& FrameRanges) const override;
 };
 
+/** Finds all the frame ranges that are looped */
+UCLASS(BlueprintType, Blueprintable, meta = (DisplayName = "Looped"))
+class UAnimDatabaseFrameRangesFunction_Looped : public UAnimDatabaseFrameRangesFunction
+{
+	GENERATED_BODY()
+
+public:
+
+	UE_API virtual FAnimDatabaseFrameRanges MakeFrameRanges_Implementation(const UAnimDatabase* Database, const FAnimDatabaseFrameRanges& FrameRanges) const override;
+};
+
+/** Finds all the frame ranges that are not looped */
+UCLASS(BlueprintType, Blueprintable, meta = (DisplayName = "Not Looped"))
+class UAnimDatabaseFrameRangesFunction_NotLooped : public UAnimDatabaseFrameRangesFunction
+{
+	GENERATED_BODY()
+
+public:
+
+	UE_API virtual FAnimDatabaseFrameRanges MakeFrameRanges_Implementation(const UAnimDatabase* Database, const FAnimDatabaseFrameRanges& FrameRanges) const override;
+};
+
 /** Finds all the frame ranges corresponding to the given Anim Sequence */
 UCLASS(BlueprintType, Blueprintable, meta = (DisplayName = "Anim Sequence"))
 class UAnimDatabaseFrameRangesFunction_AnimSequence : public UAnimDatabaseFrameRangesFunction
diff --git a/Engine/Plugins/Experimental/Animation/AnimDatabase/Source/AnimDatabaseEditor/Private/AnimDatabaseEditorMode.cpp b/Engine/Plugins/Experimental/Animation/AnimDatabase/Source/AnimDatabaseEditor/Private/AnimDatabaseEditorMode.cpp
--- a/Engine/Plugins/Experimental/Animation/AnimDatabase/Source/AnimDatabaseEditor/Private/AnimDatabaseEditorMode.cpp
+++ b/Engine/Plugins/Experimental/Animation/AnimDatabase/Source/AnimDatabaseEditor/Private/AnimDatabaseEditorMode.cpp
@@ -275,6 +275,7 @@ namespace UE::AnimDatabase::Editor
 
 			TArray<TSubclassOf<UAnimNotify>> TrackAnimNotifies;
 			TArray<TSubclassOf<UAnimNotifyState>> TrackAnimNotifyStates;
+			TArray<FName> TrackSyncMarkers;
 			TArray<FName> TrackAnimCurves;
 			TArray<FName, TInlineAllocator<32>> TrackFrames;
 			TArray<FName, TInlineAllocator<32>> TrackFrameRanges;
@@ -304,6 +305,11 @@ namespace UE::AnimDatabase::Editor
 					Database,
 					SequenceRanges);
 
+				UAnimDatabaseFrameRangesLibrary::FindSyncMarkersInFrameRanges(
+					TrackSyncMarkers,
+					Database,
+					SequenceRanges);
+
 				const int32 AnimNotifyNum = TrackAnimNotifies.Num();
 				TrackFrames.Reset();
 				TrackFramesStruct.Reset();
@@ -318,6 +324,18 @@ namespace UE::AnimDatabase::Editor
 					}
 				}
 
+				const int32 SyncMarkerNum = TrackSyncMarkers.Num();
+				for (int32 SyncMarkerIdx = 0; SyncMarkerIdx < SyncMarkerNum; SyncMarkerIdx++)
+				{
+					const FAnimDatabaseFrames SyncMarkerFramess = UAnimDatabaseFramesLibrary::MakeFramesFromSyncMarker(Database, SequenceRanges, TrackSyncMarkers[SyncMarkerIdx]);
+
+					if (SyncMarkerFramess.IsValid() && !SyncMarkerFramess.FrameSet->IsEmpty())
+					{
+						TrackFrames.Add(TrackSyncMarkers[SyncMarkerIdx]);
+						TrackFramesStruct.Add(SyncMarkerFramess);
+					}
+				}
+
 				const int32 AdditionalFramesNum = Database->Query->AdditionalTimelineFrames.Num();
 				for (int32 AdditionalIdx = 0; AdditionalIdx < AdditionalFramesNum; AdditionalIdx++)
 				{
diff --git a/Engine/Plugins/Experimental/Animation/AnimGen/Source/AnimGen/Private/AnimGenAutoEncoder.cpp b/Engine/Plugins/Experimental/Animation/AnimGen/Source/AnimGen/Private/AnimGenAutoEncoder.cpp
--- a/Engine/Plugins/Experimental/Animation/AnimGen/Source/AnimGen/Private/AnimGenAutoEncoder.cpp
+++ b/Engine/Plugins/Experimental/Animation/AnimGen/Source/AnimGen/Private/AnimGenAutoEncoder.cpp
@@ -65,11 +65,17 @@ void UAnimGenAutoEncoderTrainingSettings::ResetTransient()
 	ActiveDebugDrawer = nullptr;
 }
 
+void UAnimGenAutoEncoderTrainingSettings::LoadDatabaseAsync()
+{
+	Database.LoadAsync(FLoadSoftObjectPathAsyncDelegate::CreateLambda([this](const FSoftObjectPath&, UObject*) { ResetTransient(); }));
+}
+
 void UAnimGenAutoEncoderTrainingSettings::PostEditChangeProperty(struct FPropertyChangedEvent& Event)
 {
 	if (Event.GetMemberPropertyName() == GET_MEMBER_NAME_CHECKED(UAnimGenAutoEncoderTrainingSettings, Database))
 	{
-		Database.LoadAsync(FLoadSoftObjectPathAsyncDelegate::CreateLambda([this](const FSoftObjectPath&, UObject*) { ResetTransient(); }));
+		ResetTransient();
+		LoadDatabaseAsync();
 	}
 }
 
diff --git a/Engine/Plugins/Experimental/Animation/AnimGen/Source/AnimGen/Private/AnimGenController.cpp b/Engine/Plugins/Experimental/Animation/AnimGen/Source/AnimGen/Private/AnimGenController.cpp
--- a/Engine/Plugins/Experimental/Animation/AnimGen/Source/AnimGen/Private/AnimGenController.cpp
+++ b/Engine/Plugins/Experimental/Animation/AnimGen/Source/AnimGen/Private/AnimGenController.cpp
@@ -32,15 +32,22 @@ void UAnimGenControllerTrainingSettings::ResetTransient()
 	SelectedRanges.Empty();
 }
 
+void UAnimGenControllerTrainingSettings::LoadDatabaseAsync()
+{
+	Database.LoadAsync(FLoadSoftObjectPathAsyncDelegate::CreateLambda([this](const FSoftObjectPath&, UObject*) { ResetTransient(); }));
+}
+
 void UAnimGenControllerTrainingSettings::PostEditChangeProperty(struct FPropertyChangedEvent& Event)
 {
 	if (Event.GetMemberPropertyName() == GET_MEMBER_NAME_CHECKED(UAnimGenControllerTrainingSettings, Database))
 	{
-		Database.LoadAsync(FLoadSoftObjectPathAsyncDelegate::CreateLambda([this](const FSoftObjectPath&, UObject*) { ResetTransient(); }));
+		ResetTransient();
+		LoadDatabaseAsync();
 	}
 	else if (Event.GetMemberPropertyName() == GET_MEMBER_NAME_CHECKED(UAnimGenControllerTrainingSettings, AutoEncoder))
 	{
-		AutoEncoder.LoadAsync(FLoadSoftObjectPathAsyncDelegate::CreateLambda([this](const FSoftObjectPath&, UObject*) { ResetTransient(); }));
+		ResetTransient();
+		LoadDatabaseAsync();
 	}
 }
 
@@ -208,11 +215,16 @@ bool UAnimGenControllerTrainingSettings::Update()
 	}
 }
 
+void UAnimGenControllerTestingSettings::LoadDatabaseAsync()
+{
+	ResetDatabase.LoadAsync(FLoadSoftObjectPathAsyncDelegate());
+}
+
 void UAnimGenControllerTestingSettings::PostEditChangeProperty(struct FPropertyChangedEvent& Event)
 {
 	if (Event.GetMemberPropertyName() == GET_MEMBER_NAME_CHECKED(UAnimGenControllerTestingSettings, ResetDatabase))
 	{
-		ResetDatabase.LoadAsync(FLoadSoftObjectPathAsyncDelegate());
+		LoadDatabaseAsync();
 	}
 }
 
diff --git a/Engine/Plugins/Experimental/Animation/AnimGen/Source/AnimGen/Public/AnimGenAutoEncoder.h b/Engine/Plugins/Experimental/Animation/AnimGen/Source/AnimGen/Public/AnimGenAutoEncoder.h
--- a/Engine/Plugins/Experimental/Animation/AnimGen/Source/AnimGen/Public/AnimGenAutoEncoder.h
+++ b/Engine/Plugins/Experimental/Animation/AnimGen/Source/AnimGen/Public/AnimGenAutoEncoder.h
@@ -250,7 +250,7 @@ public:
 
 	/** Rebuilds the encoding visualization */
 	UFUNCTION(CallInEditor, Category = "Encoding Visualization", DisplayName="Rebuild Encoding Visualization")
-	UE_API void RebuildVisualizationEncoding() { bForceRebuildEncodingVisualization = true; }
+	void RebuildVisualizationEncoding() { bForceRebuildEncodingVisualization = true; }
 
 #endif
 
@@ -304,6 +304,9 @@ public:
 	UFUNCTION(CallInEditor, Category = "Data", DisplayName = "Refresh")
 	UE_API void ResetTransient();
 
+	/** Load the current database. */
+	UE_API void LoadDatabaseAsync();
+
 	/** We use this async load changes to Database */
 	UE_API virtual void PostEditChangeProperty(struct FPropertyChangedEvent& Event) override;
 
diff --git a/Engine/Plugins/Experimental/Animation/AnimGen/Source/AnimGen/Public/AnimGenController.h b/Engine/Plugins/Experimental/Animation/AnimGen/Source/AnimGen/Public/AnimGenController.h
--- a/Engine/Plugins/Experimental/Animation/AnimGen/Source/AnimGen/Public/AnimGenController.h
+++ b/Engine/Plugins/Experimental/Animation/AnimGen/Source/AnimGen/Public/AnimGenController.h
@@ -227,6 +227,9 @@ public:
 	UFUNCTION(CallInEditor, Category = "Data", DisplayName = "Refresh")
 	UE_API void ResetTransient();
 
+	/** Load the current database. */
+	UE_API void LoadDatabaseAsync();
+
 	/** We use this async load changes to Database/AutoEncoder */
 	UE_API virtual void PostEditChangeProperty(struct FPropertyChangedEvent& Event) override;
 
@@ -425,6 +428,9 @@ public:
 
 #if WITH_EDITOR
 
+	/** Load the current database. */
+	UE_API void LoadDatabaseAsync();
+
 	/** We use this async load changes to ResetDatabase */
 	UE_API virtual void PostEditChangeProperty(struct FPropertyChangedEvent& Event) override;
 
diff --git a/Engine/Plugins/Experimental/Animation/AnimGen/Source/AnimGenEditor/Private/AnimGenEditorAutoEncoderMode.cpp b/Engine/Plugins/Experimental/Animation/AnimGen/Source/AnimGenEditor/Private/AnimGenEditorAutoEncoderMode.cpp
--- a/Engine/Plugins/Experimental/Animation/AnimGen/Source/AnimGenEditor/Private/AnimGenEditorAutoEncoderMode.cpp
+++ b/Engine/Plugins/Experimental/Animation/AnimGen/Source/AnimGenEditor/Private/AnimGenEditorAutoEncoderMode.cpp
@@ -144,7 +144,10 @@ namespace UE::AnimGen::Editor
 		UAnimDatabase* QueryDatabase = AutoEncoder->TrainingSettings->Database.Get();
 		if (!QueryDatabase)
 		{
-			ViewportClient->SetWarningMessage(LOCTEXT("NoQuery", "Warning: No Training Database Provided."));
+			ViewportClient->SetWarningMessage(
+				AutoEncoder->TrainingSettings->Database.IsNull() ? 
+				LOCTEXT("NoQuery", "Warning: No Training Database Provided.") :
+				LOCTEXT("LoadingDatabase", "Loading Training Database..."));
 			ActiveRanges.Reset();
 			return;
 		}
diff --git a/Engine/Plugins/Experimental/Animation/AnimGen/Source/AnimGenEditor/Private/AnimGenEditorAutoEncoderToolkit.cpp b/Engine/Plugins/Experimental/Animation/AnimGen/Source/AnimGenEditor/Private/AnimGenEditorAutoEncoderToolkit.cpp
--- a/Engine/Plugins/Experimental/Animation/AnimGen/Source/AnimGenEditor/Private/AnimGenEditorAutoEncoderToolkit.cpp
+++ b/Engine/Plugins/Experimental/Animation/AnimGen/Source/AnimGenEditor/Private/AnimGenEditorAutoEncoderToolkit.cpp
@@ -261,7 +261,10 @@ namespace UE::AnimGen::Editor
 				.HAlign(HAlign_Center)
 				.VAlign(VAlign_Center)
 				[
-					SNew(STextBlock).Text(LOCTEXT("EmptyRanges", "No ranges returned by Query"))
+					SNew(STextBlock).Text(
+						TrainingSettings->Database.IsNull() ? LOCTEXT("NoDatabase", "No Training Database Provided") :
+						!TrainingSettings->Database.Get() ? LOCTEXT("LoadingDatabase", "Loading Training Database...") :
+						LOCTEXT("EmptyRanges", "No ranges returned by Query"))
 				];
 		}
 	}
@@ -1432,6 +1435,7 @@ namespace UE::AnimGen::Editor
 	void FAutoEncoderToolkit::InitAssetEditor(const EToolkitMode::Type Mode, const TSharedPtr<IToolkitHost>& InitToolkitHost, UAnimGenAutoEncoder* InAutoEncoder)
 	{
 		AutoEncoder = InAutoEncoder;
+		AutoEncoder->TrainingSettings->LoadDatabaseAsync();
 
 		// Create Preview Scene
 		PreviewScene = MakeShared<AnimDatabase::Editor::FPreviewScene>(AnimDatabase::Editor::FPreviewScene::ConstructionValues().ShouldSimulatePhysics(true).ForceUseMovementComponentInNonGameWorld(true));
diff --git a/Engine/Plugins/Experimental/Animation/AnimGen/Source/AnimGenEditor/Private/AnimGenEditorControllerMode.cpp b/Engine/Plugins/Experimental/Animation/AnimGen/Source/AnimGenEditor/Private/AnimGenEditorControllerMode.cpp
--- a/Engine/Plugins/Experimental/Animation/AnimGen/Source/AnimGenEditor/Private/AnimGenEditorControllerMode.cpp
+++ b/Engine/Plugins/Experimental/Animation/AnimGen/Source/AnimGenEditor/Private/AnimGenEditorControllerMode.cpp
@@ -522,7 +522,15 @@ namespace UE::AnimGen::Editor
 		// Load Database
 
 		UAnimDatabase* TrainingDatabase = Controller->TrainingSettings->Database.Get();
-		if (!TrainingDatabase) { return; }
+		if (!TrainingDatabase)
+		{
+			ViewportClient->SetWarningMessage(
+				Controller->TrainingSettings->Database.IsNull() ?
+				LOCTEXT("NoQuery", "Warning: No Training Database Provided.") :
+				LOCTEXT("LoadingDatabase", "Loading Training Database..."));
+			ActiveRanges.Reset();
+			return;
+		}
 
 		// Load AutoEncoder
 
diff --git a/Engine/Plugins/Experimental/Animation/AnimGen/Source/AnimGenEditor/Private/AnimGenEditorControllerToolkit.cpp b/Engine/Plugins/Experimental/Animation/AnimGen/Source/AnimGenEditor/Private/AnimGenEditorControllerToolkit.cpp
--- a/Engine/Plugins/Experimental/Animation/AnimGen/Source/AnimGenEditor/Private/AnimGenEditorControllerToolkit.cpp
+++ b/Engine/Plugins/Experimental/Animation/AnimGen/Source/AnimGenEditor/Private/AnimGenEditorControllerToolkit.cpp
@@ -259,7 +259,10 @@ namespace UE::AnimGen::Editor
 				.HAlign(HAlign_Center)
 				.VAlign(VAlign_Center)
 				[
-					SNew(STextBlock).Text(LOCTEXT("EmptyRanges", "No Control Sets returned by Behavior"))
+					SNew(STextBlock).Text(
+						TrainingSettings->Database.IsNull() ? LOCTEXT("NoDatabase", "No Training Database Provided") :
+						!TrainingSettings->Database.Get() ? LOCTEXT("LoadingDatabase", "Loading Training Database...") :
+						LOCTEXT("EmptyRanges", "No Control Sets returned by Behavior"))
 				];
 		}
 	}
@@ -1616,8 +1619,8 @@ namespace UE::AnimGen::Editor
 	void FControllerToolkit::InitAssetEditor(const EToolkitMode::Type Mode, const TSharedPtr<IToolkitHost>& InitToolkitHost, UAnimGenController* InController)
 	{
 		Controller = InController;
-		Controller->TrainingSettings->ResetTransient();
-		Controller->TrainingSettings->Update();
+		Controller->TrainingSettings->LoadDatabaseAsync();
+		Controller->TestingSettings->LoadDatabaseAsync();
 
 		PreviewScene = MakeShared<AnimDatabase::Editor::FPreviewScene>(AnimDatabase::Editor::FPreviewScene::ConstructionValues().ShouldSimulatePhysics(true).ForceUseMovementComponentInNonGameWorld(true));
 		//Temporary fix for missing attached assets - MDW (Copied from FPersonaToolkit::CreatePreviewScene)
