AnimDatabase: Added some more frame attributes

#Synced-CL 52006914

diff --git a/Engine/Plugins/Experimental/Animation/AnimDatabase/Source/AnimDatabase/Private/AnimDatabase.cpp b/Engine/Plugins/Experimental/Animation/AnimDatabase/Source/AnimDatabase/Private/AnimDatabase.cpp
--- a/Engine/Plugins/Experimental/Animation/AnimDatabase/Source/AnimDatabase/Private/AnimDatabase.cpp
+++ b/Engine/Plugins/Experimental/Animation/AnimDatabase/Source/AnimDatabase/Private/AnimDatabase.cpp
@@ -2357,6 +2357,9 @@ void UAnimDatabaseQuery::ResetTransient()
 	AdditionalTimelineFramesObjects.Reset();
 	AdditionalTimelineFrameRangesObjects.Reset();
 	AdditionalTimelineFrameAttributeObjects.Reset();
+	ActiveAdditionalTimelineFrames.Reset();
+	ActiveAdditionalTimelineFrameRanges.Reset();
+	ActiveAdditionalTimelineFrameAttributes.Reset();
 	ActiveDebugDrawer = nullptr;
 }
 
@@ -7685,6 +7688,62 @@ void UAnimDatabaseDebugDraw_RootOrientation::DrawDebug_Implementation(
 		Radius);
 }
 
+void UAnimDatabaseDebugDraw_MovementDirection::DrawDebug_Implementation(
+	const FDebugDrawer& Drawer,
+	const FAnimDatabasePoseState& InPoseState,
+	const UAnimDatabase* InDatabase,
+	const FAnimDatabaseFrameRanges& InFrameRanges,
+	const FTransform& RangeViewportTransform,
+	const int32 SequenceIdx,
+	const float SequenceTime,
+	const int32 RangeStart,
+	const int32 RangeLength,
+	const FLinearColor& IdentifierColor)
+{
+	const FVector RootLocation = UAnimDatabasePoseStateLibrary::PoseStateRootLocation(InPoseState) + FVector(0.0, 0.0, FloorOffset);
+	const FVector RootDirection = UAnimDatabasePoseStateLibrary::PoseStateRootDirection(InPoseState, Direction);
+	const FVector RootLinearVelocity = UAnimDatabasePoseStateLibrary::PoseStateRootLinearVelocity(InPoseState);
+
+	UDrawDebugLibrary::DrawDebugLine(
+		Drawer,
+		RootLocation,
+		RootLocation + VelocityScale * RootLinearVelocity,
+		UDrawDebugLibrary::MakeDrawDebugLineStyleFromColorAndThickness(IdentifierColor, Thickness),
+		true);
+
+	UDrawDebugLibrary::DrawDebugArrow(
+		Drawer,
+		RootLocation,
+		RootLocation + DirectionArrowLength * RootDirection,
+		UDrawDebugLibrary::MakeDrawDebugLineStyleFromColorAndThickness(IdentifierColor, Thickness),
+		true);
+
+	if (RootLinearVelocity.SquaredLength() > FMath::Square(VelocityThreshold))
+	{
+		const float Angle = FMath::Acos(RootDirection.Dot(RootLinearVelocity.GetSafeNormal())) * (RootDirection.Cross(RootLinearVelocity).Z > 0.0f ? 1.0f : -1.0f);
+
+		UDrawDebugLibrary::DrawDebugArc(
+			Drawer,
+			RootLocation,
+			RootDirection.Rotation(),
+			FMath::RadiansToDegrees(Angle),
+			UDrawDebugLibrary::MakeDrawDebugLineStyleFromColorAndThickness(IdentifierColor, Thickness));
+
+		FDrawDebugStringSettings StringSettings;
+		StringSettings.bCenterHorizontally = true;
+		StringSettings.bCenterVertically = true;
+
+		UDrawDebugLibrary::DrawDebugStringView(
+			Drawer,
+			FString::Printf(TEXT("Angle: % 5.1f deg"), FMath::RadiansToDegrees(Angle)),
+			RootLocation + DirectionArrowLength * RootDirection + FVector(0.0f, 0.0f, 20.0f),
+			(-RootDirection).Rotation(),
+			UDrawDebugLibrary::MakeDrawDebugLineStyleFromColorAndThickness(IdentifierColor),
+			true,
+			StringSettings);
+	}
+}
+
 void UAnimDatabaseDebugDraw_BoneTransform::DrawDebug_Implementation(
 	const FDebugDrawer& Drawer,
 	const FAnimDatabasePoseState& InPoseState,
diff --git a/Engine/Plugins/Experimental/Animation/AnimDatabase/Source/AnimDatabase/Private/AnimDatabaseFrameAttribute.cpp b/Engine/Plugins/Experimental/Animation/AnimDatabase/Source/AnimDatabase/Private/AnimDatabaseFrameAttribute.cpp
--- a/Engine/Plugins/Experimental/Animation/AnimDatabase/Source/AnimDatabase/Private/AnimDatabaseFrameAttribute.cpp
+++ b/Engine/Plugins/Experimental/Animation/AnimDatabase/Source/AnimDatabase/Private/AnimDatabaseFrameAttribute.cpp
@@ -8904,6 +8904,27 @@ FAnimDatabaseFrameAttribute UAnimDatabaseFrameAttributeLibrary::FrameAttributeLo
 	return Out;
 }
 
+FAnimDatabaseFrameAttribute UAnimDatabaseFrameAttributeLibrary::FrameAttributeVelocityToDirection(const FAnimDatabaseFrameAttribute& VelocityFrameAttribute)
+{
+	if (!VelocityFrameAttribute.IsValid())
+	{
+		UE_LOGFMT(LogAnimDatabase, Error, "FrameAttributeVelocityToDirection: Invalid FrameAttribute.");
+		return FAnimDatabaseFrameAttribute();
+	}
+
+	if (VelocityFrameAttribute.Type != EAnimDatabaseAttributeType::LinearVelocity &&
+		VelocityFrameAttribute.Type != EAnimDatabaseAttributeType::AngularVelocity)
+	{
+		UE_LOGFMT(LogAnimDatabase, Error, "FrameAttributeVelocityToDirection: FrameAttribute Types don't match. Got {Type}. Expected LinearVelocity or AngularVelocity.", AttributeTypeNameInternal(VelocityFrameAttribute.Type));
+		return FAnimDatabaseFrameAttribute();
+	}
+
+	FAnimDatabaseFrameAttribute Out = FrameAttributeCopy(VelocityFrameAttribute);
+	Out.Type = EAnimDatabaseAttributeType::Direction;
+	UE::Learning::FrameAttribute::NormalizeInplace(*Out.FrameAttribute);
+	return Out;
+}
+
 FAnimDatabaseFrameAttribute UAnimDatabaseFrameAttributeLibrary::FrameAttributeFilterGaussian(const FAnimDatabaseFrameAttribute& A, const float StandardDeviationInFrames)
 {
 	if (!A.IsValid())
@@ -9039,6 +9060,11 @@ FAnimDatabaseFrameRanges UAnimDatabaseFrameAttributeLibrary::FrameAttributeBoolT
 	return Out;
 }
 
+FAnimDatabaseFrameAttribute UAnimDatabaseFrameAttributeLibrary::FrameAttributeFrameRangesToBool(const FAnimDatabaseFrameRanges& ActiveFrameRanges, const FAnimDatabaseFrameRanges& InFrameRanges)
+{
+	return MakeBoolFrameAttributeFromActiveRanges(ActiveFrameRanges, InFrameRanges);
+}
+
 FAnimDatabaseFrameAttribute UAnimDatabaseFrameAttributeLibrary::FrameAttributeSelect(const FAnimDatabaseFrameAttribute& Cond, const FAnimDatabaseFrameAttribute& True, const FAnimDatabaseFrameAttribute& False)
 {
 	TRACE_CPUPROFILER_EVENT_SCOPE(UAnimDatabaseFrameAttributeLibrary::FrameAttributeSelect);
@@ -10084,7 +10110,9 @@ FAnimDatabaseFrames UAnimDatabaseFrameAttributeLibrary::FrameAttributeContactFra
 	const float HeightThreshold,
 	const float VelocityThreshold,
 	const bool bFilterCurve,
-	const int32 MajorityVoteFilterWidth)
+	const int32 MajorityVoteFilterWidth,
+	const float FrameOffset,
+	const bool bIgnoreFirstFrame)
 {
 	if (!Database)
 	{
@@ -10101,35 +10129,14 @@ FAnimDatabaseFrames UAnimDatabaseFrameAttributeLibrary::FrameAttributeContactFra
 
 	if (bFilterCurve) { ContactActive = FrameAttributeFilterMajorityVote(ContactActive, MajorityVoteFilterWidth); }
 
-	return UAnimDatabaseFrameRangesLibrary::MakeFramesAtFrameRangesStarts(FrameAttributeBoolToFrameRanges(ContactActive));
-}
+	FAnimDatabaseFrames ContactFrames = UAnimDatabaseFrameRangesLibrary::MakeFramesAtFrameRangesStarts(FrameAttributeBoolToFrameRanges(ContactActive));
 
-FAnimDatabaseFrameAttribute UAnimDatabaseFrameAttributeLibrary::FrameAttributeContactEvent(
-	const UAnimDatabase* Database,
-	const FAnimDatabaseFrameRanges& FrameRanges,
-	const FName BoneName,
-	const float HeightThreshold,
-	const float VelocityThreshold,
-	const bool bFilterCurve,
-	const int32 MajorityVoteFilterWidth,
-	const float EventOffset)
-{
-	if (!Database)
+	if (bIgnoreFirstFrame)
 	{
-		UE_LOGFMT(LogAnimDatabase, Error, "FrameAttributeContactEvent: Database is nullptr.");
-		return FAnimDatabaseFrameAttribute();
+		ContactFrames = UAnimDatabaseFramesLibrary::FramesDifference(ContactFrames, UAnimDatabaseFrameRangesLibrary::MakeFramesAtFrameRangesStarts(FrameRanges));
 	}
 
-	return MakeEventFrameAttribute(UAnimDatabaseFramesLibrary::FramesBefore(Database, FrameAttributeContactFrames(
-			Database,
-			FrameRanges,
-			BoneName,
-			HeightThreshold,
-			VelocityThreshold,
-			bFilterCurve,
-			MajorityVoteFilterWidth), FMath::RoundToInt(EventOffset * Database->GetFrameRate().AsDecimal()), EAnimDatabaseFrameShiftBehavior::Clamp),
-		FrameRanges,
-		Database->GetFrameRate());
+	return UAnimDatabaseFramesLibrary::FramesBefore(Database, ContactFrames, FMath::RoundToInt(FrameOffset * Database->GetFrameRate().AsDecimal()), EAnimDatabaseFrameShiftBehavior::Clamp);
 }
 
 FAnimDatabaseFrameAttribute UAnimDatabaseFrameAttributeLibrary::FrameAttributeFitTransform(
@@ -10258,6 +10265,11 @@ FAnimDatabaseFrameAttribute UAnimDatabaseFrameAttributeFunction::MakeFrameAttrib
 	return UAnimDatabaseFrameAttributeLibrary::MakeEmptyFrameAttribute();
 }
 
+FAnimDatabaseFrameAttribute UAnimDatabaseFrameAttributeFunction_FrameRanges::MakeFrameAttribute_Implementation(const UAnimDatabase* Database, const FAnimDatabaseFrameRanges& FrameRanges) const
+{
+	return UAnimDatabaseFrameAttributeLibrary::MakeBoolFrameAttributeFromActiveRanges(UAnimDatabaseFrameRangesLibrary::MakeFrameRangesFromFunction(Database, FrameRanges, FrameRangesFunction), FrameRanges);
+}
+
 FAnimDatabaseFrameAttribute UAnimDatabaseFrameAttributeFunction_Filter::MakeFrameAttribute_Implementation(const UAnimDatabase* Database, const FAnimDatabaseFrameRanges& FrameRanges) const
 {
 	return UAnimDatabaseFrameAttributeLibrary::MakeFrameAttributeFromFunction(Database, UAnimDatabaseFrameRangesLibrary::MakeFrameRangesFromFunction(Database, FrameRanges, FrameRangesFunction), FrameAttributeFunction);
@@ -10296,7 +10308,7 @@ FAnimDatabaseFrameAttribute UAnimDatabaseFrameAttributeFunction_ContactCurve::Ma
 
 FAnimDatabaseFrameAttribute UAnimDatabaseFrameAttributeFunction_ContactEvent::MakeFrameAttribute_Implementation(const UAnimDatabase* Database, const FAnimDatabaseFrameRanges& FrameRanges) const
 {
-	return UAnimDatabaseFrameAttributeLibrary::FrameAttributeContactEvent(
+	FAnimDatabaseFrames ContactFrames = UAnimDatabaseFrameAttributeLibrary::FrameAttributeContactFrames(
 		Database,
 		FrameRanges,
 		BoneName,
@@ -10305,6 +10317,13 @@ FAnimDatabaseFrameAttribute UAnimDatabaseFrameAttributeFunction_ContactEvent::Ma
 		bFilterCurve,
 		MajorityVoteFilterWidth,
 		EventOffset);
+
+	if (EventFilter)
+	{
+		ContactFrames = UAnimDatabaseFrameRangesLibrary::FramesFrameRangesIntersection(ContactFrames, UAnimDatabaseFrameRangesLibrary::MakeFrameRangesFromFunction(Database, FrameRanges, EventFilter));
+	}
+
+	return UAnimDatabaseFrameAttributeLibrary::MakeEventFrameAttribute(ContactFrames, FrameRanges, Database->GetFrameRate());
 }
 
 FAnimDatabaseFrameAttribute UAnimDatabaseFrameAttributeFunction_FitTransform::MakeFrameAttribute_Implementation(const UAnimDatabase* Database, const FAnimDatabaseFrameRanges& FrameRanges) const
@@ -10359,6 +10378,26 @@ FAnimDatabaseFrameAttribute UAnimDatabaseFrameAttributeFunction_BoneGlobalLinear
 	return UAnimDatabaseFrameAttributeLibrary::FrameAttributeBoneSpeed(Database, FrameRanges, BoneName);
 }
 
+FAnimDatabaseFrameRanges UAnimDatabaseFrameRangesFunction_MajorityVoteFilter::MakeFrameRanges_Implementation(const UAnimDatabase* Database, const FAnimDatabaseFrameRanges& FrameRanges) const
+{
+	return UAnimDatabaseFrameAttributeLibrary::FrameAttributeBoolToFrameRanges(
+		UAnimDatabaseFrameAttributeLibrary::FrameAttributeFilterMajorityVote(
+			UAnimDatabaseFrameAttributeLibrary::FrameAttributeFrameRangesToBool(UAnimDatabaseFrameRangesLibrary::MakeFrameRangesFromFunction(Database, FrameRanges, FilterFrameRanges), FrameRanges),
+			MajorityVoteFilterWidth));
+}
+
+FAnimDatabaseFrameRanges UAnimDatabaseFrameRangesFunction_MovingInDirection::MakeFrameRanges_Implementation(const UAnimDatabase* Database, const FAnimDatabaseFrameRanges& FrameRanges) const
+{
+	const FAnimDatabaseFrameAttribute MovingDirection = UAnimDatabaseFrameAttributeLibrary::MakeRootDirectionFrameAttribute(Database, FrameRanges, 0.0f, Direction);
+	const FAnimDatabaseFrameAttribute RootVelocity = UAnimDatabaseFrameAttributeLibrary::MakeRootLinearVelocityFrameAttribute(Database, FrameRanges);
+	const FAnimDatabaseFrameAttribute RootVelocityDirection = UAnimDatabaseFrameAttributeLibrary::FrameAttributeVelocityToDirection(RootVelocity);
+
+	return UAnimDatabaseFrameAttributeLibrary::FrameAttributeBoolToFrameRanges(
+		UAnimDatabaseFrameAttributeLibrary::FrameAttributeAnd(
+			UAnimDatabaseFrameAttributeLibrary::FrameAttributeGreaterThanFloat(UAnimDatabaseFrameAttributeLibrary::FrameAttributeDot(MovingDirection, RootVelocityDirection), FMath::Cos(FMath::DegreesToRadians(AngleThreshold))),
+			UAnimDatabaseFrameAttributeLibrary::FrameAttributeGreaterThanFloat(UAnimDatabaseFrameAttributeLibrary::FrameAttributeLength(RootVelocity), VelocityThreshold)));
+}
+
 FAnimDatabaseFrameRanges UAnimDatabaseFrameRangesFunction_RootLinearVelocityBelowThreshold::MakeFrameRanges_Implementation(const UAnimDatabase* Database, const FAnimDatabaseFrameRanges& FrameRanges) const
 {
 	return UAnimDatabaseFrameAttributeLibrary::FrameAttributeBoolToFrameRanges(UAnimDatabaseFrameAttributeLibrary::FrameAttributeLessThanFloat(
diff --git a/Engine/Plugins/Experimental/Animation/AnimDatabase/Source/AnimDatabase/Public/AnimDatabase.h b/Engine/Plugins/Experimental/Animation/AnimDatabase/Source/AnimDatabase/Public/AnimDatabase.h
--- a/Engine/Plugins/Experimental/Animation/AnimDatabase/Source/AnimDatabase/Public/AnimDatabase.h
+++ b/Engine/Plugins/Experimental/Animation/AnimDatabase/Source/AnimDatabase/Public/AnimDatabase.h
@@ -2447,11 +2447,11 @@ class UAnimDatabaseDebugDraw_RootOrientation : public UAnimDatabaseDebugDraw
 public:
 
 	/** Vertical Offset from Floor */
-	UPROPERTY(EditAnywhere, Category = "Settings", meta = (ForceForceUnits = "cm"))
+	UPROPERTY(EditAnywhere, Category = "Settings", meta = (ForceUnits = "cm"))
 	float FloorOffset = 1.0f;
 
 	/** Widget Radius */
-	UPROPERTY(EditAnywhere, Category = "Settings", meta = (ForceForceUnits = "cm"))
+	UPROPERTY(EditAnywhere, Category = "Settings", meta = (ForceUnits = "cm"))
 	float Radius = 50.0f;
 
 	/** Local root forward direction */
@@ -2479,6 +2479,61 @@ public:
 #endif
 };
 
+/** Debug Draw class for drawing the current movement direction */
+UCLASS(BlueprintType, Blueprintable, meta = (DisplayName = "Movement Direction"))
+class UAnimDatabaseDebugDraw_MovementDirection : public UAnimDatabaseDebugDraw
+{
+	GENERATED_BODY()
+
+#if WITH_EDITORONLY_DATA
+
+public:
+
+	/** Vertical Offset from Floor */
+	UPROPERTY(EditAnywhere, Category = "Settings", meta = (ForceUnits = "cm"))
+	float FloorOffset = 1.0f;
+
+	/** Local root forward direction */
+	UPROPERTY(EditAnywhere, Category = "Settings")
+	FVector Direction = FVector::RightVector;
+
+	/** Character is only considered moving when the velocity is above this threshold */
+	UPROPERTY(EditAnywhere, Category = "Settings", meta = (ClampMin = "0", UIMin = "0", ForceUnits = "cm/s"))
+	float VelocityThreshold = 25.0f;
+
+	/** Direction Arrow Length */
+	UPROPERTY(EditAnywhere, Category = "Settings", meta = (ForceUnits = "cm"))
+	float DirectionArrowLength = 50.0f;
+
+	/** Velocity Scale */
+	UPROPERTY(EditAnywhere, Category = "Settings")
+	float VelocityScale = 0.5f;
+
+	/** Widget Thickness */
+	UPROPERTY(EditAnywhere, Category = "Settings", meta = (ForceUnits = "cm"))
+	float Thickness = 1.0f;
+
+#endif
+
+#if WITH_EDITOR
+
+public:
+
+	virtual void DrawDebug_Implementation(
+		const FDebugDrawer& Drawer,
+		const FAnimDatabasePoseState& InPoseState,
+		const UAnimDatabase* InDatabase,
+		const FAnimDatabaseFrameRanges& InFrameRanges,
+		const FTransform& RangeViewportTransform,
+		const int32 SequenceIdx,
+		const float SequenceTime,
+		const int32 RangeStart,
+		const int32 RangeLength,
+		const FLinearColor& IdentifierColor) override;
+
+#endif
+};
+
 /** Debug Draw the transform of the given bone */
 UCLASS(BlueprintType, Blueprintable, meta = (DisplayName = "Bone Transform"))
 class UAnimDatabaseDebugDraw_BoneTransform : public UAnimDatabaseDebugDraw
@@ -2494,11 +2549,11 @@ public:
 	FName BoneName = NAME_None;
 
 	/** Widget Radius */
-	UPROPERTY(EditAnywhere, Category = "Settings", meta = (ForceForceUnits = "cm"))
+	UPROPERTY(EditAnywhere, Category = "Settings", meta = (ForceUnits = "cm"))
 	float Radius = 25.0f;
 
 	/** Widget Thickness */
-	UPROPERTY(EditAnywhere, Category = "Settings", meta = (ForceForceUnits = "cm"))
+	UPROPERTY(EditAnywhere, Category = "Settings", meta = (ForceUnits = "cm"))
 	float Thickness = 1.0f;
 
 	/** If to depth test while drawing */
@@ -2550,11 +2605,11 @@ public:
 	bool bDrawTransform = true;
 
 	/** Transform Radius */
-	UPROPERTY(EditAnywhere, Category = "Settings", meta = (ForceForceUnits = "cm"))
+	UPROPERTY(EditAnywhere, Category = "Settings", meta = (ForceUnits = "cm"))
 	float TransformRadius = 25.0f;
 
 	/** Transform Thickness */
-	UPROPERTY(EditAnywhere, Category = "Settings", meta = (ForceForceUnits = "cm"))
+	UPROPERTY(EditAnywhere, Category = "Settings", meta = (ForceUnits = "cm"))
 	float TransformThickness = 1.0f;
 
 	/** Static mesh to use for attachment (will draw bounding box only) */
@@ -2562,7 +2617,7 @@ public:
 	TObjectPtr<UStaticMesh> StaticMesh;
 
 	/** Bounding box Thickness */
-	UPROPERTY(EditAnywhere, Category = "Settings", meta = (ForceForceUnits = "cm"))
+	UPROPERTY(EditAnywhere, Category = "Settings", meta = (ForceUnits = "cm"))
 	float BoundingBoxThickness = 1.0f;
 
 	/** If to depth test while drawing */
@@ -2601,7 +2656,7 @@ class UAnimDatabaseDebugDraw_RootVelocity : public UAnimDatabaseDebugDraw
 public:
 
 	/** Graph axes maximum for linear velocity */
-	UPROPERTY(EditAnywhere, Category = "Settings", meta = (ForceForceUnits = "cm/s"))
+	UPROPERTY(EditAnywhere, Category = "Settings", meta = (ForceUnits = "cm/s"))
 	float MaxLinearVelocity = 500.0f;
 
 	/** Graph axes maximum for angular velocity */
diff --git a/Engine/Plugins/Experimental/Animation/AnimDatabase/Source/AnimDatabase/Public/AnimDatabaseFrameAttribute.h b/Engine/Plugins/Experimental/Animation/AnimDatabase/Source/AnimDatabase/Public/AnimDatabaseFrameAttribute.h
--- a/Engine/Plugins/Experimental/Animation/AnimDatabase/Source/AnimDatabase/Public/AnimDatabaseFrameAttribute.h
+++ b/Engine/Plugins/Experimental/Animation/AnimDatabase/Source/AnimDatabase/Public/AnimDatabaseFrameAttribute.h
@@ -1246,6 +1246,10 @@ public:
 	UFUNCTION(BlueprintPure, Category = "AnimDatabase")
 	static UE_API FAnimDatabaseFrameRanges FrameAttributeBoolToFrameRanges(const FAnimDatabaseFrameAttribute& BoolFrameAttribute);
 
+	/** Converts frame ranges to a Bool frame attribute  */
+	UFUNCTION(BlueprintPure, Category = "AnimDatabase")
+	static UE_API FAnimDatabaseFrameAttribute FrameAttributeFrameRangesToBool(const FAnimDatabaseFrameRanges& ActiveFrameRanges, const FAnimDatabaseFrameRanges& InFrameRanges);
+
 	/** Selects between True and False based on the given condition bool frame attribute */
 	UFUNCTION(BlueprintPure, Category = "AnimDatabase")
 	static UE_API FAnimDatabaseFrameAttribute FrameAttributeSelect(const FAnimDatabaseFrameAttribute& Cond, const FAnimDatabaseFrameAttribute& True, const FAnimDatabaseFrameAttribute& False);
@@ -1681,6 +1685,10 @@ public:
 	UFUNCTION(BlueprintPure, Category = "AnimDatabase")
 	static UE_API FAnimDatabaseFrameAttribute FrameAttributeLocationToDirection(const FAnimDatabaseFrameAttribute& LocationFrameAttribute);
 
+	/** Converts a velocity frame attribute to a direction frame attribute (normalizing the result) */
+	UFUNCTION(BlueprintPure, Category = "AnimDatabase")
+	static UE_API FAnimDatabaseFrameAttribute FrameAttributeVelocityToDirection(const FAnimDatabaseFrameAttribute& VelocityFrameAttribute);
+
 public:
 
 	/** Applies a Gaussian smoothing filter to the values in each range of a frame attribute. Assumes type is compatible */
@@ -1896,32 +1904,11 @@ public:
 	 * @param bFilterCurve						If to apply the majority vote filter to the curve
 	 * @param MajorityVoteFilterWidth			Filter width (in frames) for the majority vote filter used to remove very short contact transitions.
 	 *											contact periods (or non-contact periods) less than half this duration will be removed.
+	 * @param FrameOffset						The amount of time before the contact (in seconds) to have the frame
+	 * @param bIgnoreFirstFrame					Removes contact frames that occur on the first frame of the frame ranges
 	 */
 	UFUNCTION(BlueprintPure, Category = "AnimDatabase", meta = (NotBlueprintThreadSafe))
 	static UE_API FAnimDatabaseFrames FrameAttributeContactFrames(
-		const UAnimDatabase* Database,
-		const FAnimDatabaseFrameRanges& FrameRanges,
-		const FName BoneName = TEXT("ball_l"),
-		const float HeightThreshold = 20.0f,
-		const float VelocityThreshold = 25.0f,
-		const bool bFilterCurve = true,
-		const int32 MajorityVoteFilterWidth = 5);
-
-	/**
-	 * Computes an event attribute representing a contact down event for the given joint.
-	 *
-	 * @param Database							Database to use
-	 * @param FrameRanges						Ranges in the database to use
-	 * @param BoneName							Name of the bone to compute the contact state for
-	 * @param HeightThreshold					Maximum height (in cm) for which the bone can be considered in contact
-	 * @param VelocityThreshold					Maximum velocity (in cm/s) for which the bone can be considered in contact
-	 * @param bFilterCurve						If to apply the majority vote filter to the curve
-	 * @param MajorityVoteFilterWidth			Filter width (in frames) for the majority vote filter used to remove very short contact transitions.
-	 *											contact periods (or non-contact periods) less than half this duration will be removed.
-	 * @param EventOffset						The amount of time before the contact (in seconds) to have the event
-	 */
-	UFUNCTION(BlueprintPure, Category = "AnimDatabase", meta = (NotBlueprintThreadSafe))
-	static UE_API FAnimDatabaseFrameAttribute FrameAttributeContactEvent(
 		const UAnimDatabase* Database,
 		const FAnimDatabaseFrameRanges& FrameRanges,
 		const FName BoneName = TEXT("ball_l"),
@@ -1929,7 +1916,8 @@ public:
 		const float VelocityThreshold = 25.0f,
 		const bool bFilterCurve = true,
 		const int32 MajorityVoteFilterWidth = 5,
-		const float EventOffset = 0.0f);
+		const float FrameOffset = 0.0f,
+		const bool bIgnoreFirstFrame = true);
 
 	/**
 	 * Fits a transform from the given location and direction attributes
@@ -1984,6 +1972,23 @@ public:
 	UE_API FAnimDatabaseFrameAttribute MakeFrameAttribute(const UAnimDatabase* Database, const FAnimDatabaseFrameRanges& FrameRanges) const;
 };
 
+/** Computes a bool Frame Attribute for the given Frame Ranges  */
+UCLASS(BlueprintType, Blueprintable, meta = (DisplayName = "Frame Ranges"))
+class UAnimDatabaseFrameAttributeFunction_FrameRanges : public UAnimDatabaseFrameAttributeFunction
+{
+	GENERATED_BODY()
+
+public:
+
+	/** Frame Ranges to use */
+	UPROPERTY(EditAnywhere, Instanced, Category = "Settings")
+	TObjectPtr<UAnimDatabaseFrameRangesFunction> FrameRangesFunction;
+
+public:
+
+	UE_API virtual FAnimDatabaseFrameAttribute MakeFrameAttribute_Implementation(const UAnimDatabase* Database, const FAnimDatabaseFrameRanges& FrameRanges) const override;
+};
+
 /** Computes an Attribute for the given Frame Ranges  */
 UCLASS(BlueprintType, Blueprintable, meta = (DisplayName = "Filter"))
 class UAnimDatabaseFrameAttributeFunction_Filter : public UAnimDatabaseFrameAttributeFunction
@@ -2130,9 +2135,13 @@ public:
 	int32 MajorityVoteFilterWidth = 5;
 
 	/** The amount of time before the contact (in seconds) to have the event */
-	UPROPERTY(EditAnywhere, Category = "Filtering", meta = (ForceUnits = "s"))
+	UPROPERTY(EditAnywhere, Category = "Settings", meta = (ForceUnits = "s"))
 	float EventOffset = 0.0f;
 
+	/** If provided, will be used to filter the contact events */
+	UPROPERTY(EditAnywhere, Instanced, Category = "Settings")
+	TObjectPtr<UAnimDatabaseFrameRangesFunction> EventFilter;
+
 public:
 
 	UE_API virtual FAnimDatabaseFrameAttribute MakeFrameAttribute_Implementation(const UAnimDatabase* Database, const FAnimDatabaseFrameRanges& FrameRanges) const override;
@@ -2284,6 +2293,54 @@ public:
 	UE_API virtual FAnimDatabaseFrameAttribute MakeFrameAttribute_Implementation(const UAnimDatabase* Database, const FAnimDatabaseFrameRanges& FrameRanges) const override;
 };
 
+/** Filter frame ranges with a majority vote filter */
+UCLASS(BlueprintType, Blueprintable, meta = (DisplayName = "Majority Vote Filter"))
+class UAnimDatabaseFrameRangesFunction_MajorityVoteFilter : public UAnimDatabaseFrameRangesFunction
+{
+	GENERATED_BODY()
+
+public:
+
+	/** Frame ranges to filter */
+	UPROPERTY(EditAnywhere, Instanced, Category = "Settings")
+	TObjectPtr<UAnimDatabaseFrameRangesFunction> FilterFrameRanges;
+
+	/**
+	 * Filter width (in frames) for the majority vote filter used to remove very short frame ranges less than half this duration will be removed.
+	 */
+	UPROPERTY(EditAnywhere, Category = "Settings")
+	int32 MajorityVoteFilterWidth = 15;
+
+public:
+
+	UE_API virtual FAnimDatabaseFrameRanges MakeFrameRanges_Implementation(const UAnimDatabase* Database, const FAnimDatabaseFrameRanges& FrameRanges) const override;
+};
+
+/** Frame ranges where the character is moving in the given direction */
+UCLASS(BlueprintType, Blueprintable, meta = (DisplayName = "Moving in Direction"))
+class UAnimDatabaseFrameRangesFunction_MovingInDirection : public UAnimDatabaseFrameRangesFunction
+{
+	GENERATED_BODY()
+
+public:
+
+	/** Local direction to check if the character is moving in */
+	UPROPERTY(EditAnywhere, Category = "Settings")
+	FVector Direction = FVector::RightVector;
+
+	/** Character is only considered moving when the velocity is above this threshold */
+	UPROPERTY(EditAnywhere, Category = "Settings", meta = (ClampMin = "0", UIMin = "0", ForceUnits = "cm/s"))
+	float VelocityThreshold = 25.0f;
+
+	/** Character is considered moving in the direction when the angle between the velocity direction and the specified direction is less than this */
+	UPROPERTY(EditAnywhere, Category = "Settings", meta = (ClampMin = "0", UIMin = "0", ForceUnits = "deg"))
+	float AngleThreshold = 45.0f;
+
+public:
+
+	UE_API virtual FAnimDatabaseFrameRanges MakeFrameRanges_Implementation(const UAnimDatabase* Database, const FAnimDatabaseFrameRanges& FrameRanges) const override;
+};
+
 /** Frame ranges where the root linear velocity magnitude is below some threshold */
 UCLASS(BlueprintType, Blueprintable, meta = (DisplayName = "Root Linear Velocity Below Threshold"))
 class UAnimDatabaseFrameRangesFunction_RootLinearVelocityBelowThreshold : public UAnimDatabaseFrameRangesFunction
@@ -2292,7 +2349,7 @@ class UAnimDatabaseFrameRangesFunction_RootLinearVelocityBelowThreshold : public
 
 public:
 
-	UPROPERTY(EditAnywhere, Category = "Settings", meta = (ClampMin = "0", UIMin = "0", Units = "cm/s"))
+	UPROPERTY(EditAnywhere, Category = "Settings", meta = (ClampMin = "0", UIMin = "0", ForceUnits = "cm/s"))
 	float Threshold = 10.0f;
 
 public:
@@ -2308,7 +2365,7 @@ class UAnimDatabaseFrameRangesFunction_RootLinearVelocityAboveThreshold : public
 
 public:
 
-	UPROPERTY(EditAnywhere, Category = "Settings", meta = (ClampMin = "0", UIMin = "0", Units = "cm/s"))
+	UPROPERTY(EditAnywhere, Category = "Settings", meta = (ClampMin = "0", UIMin = "0", ForceUnits = "cm/s"))
 	float Threshold = 10.0f;
 
 public:
