AnimDatabase: Some UX improvements and functions for automatic phase annotation.

#Synced-CL 52098865

diff --git a/Engine/Plugins/Experimental/Animation/AnimDatabase/Source/AnimDatabase/Private/AnimDatabase.cpp b/Engine/Plugins/Experimental/Animation/AnimDatabase/Source/AnimDatabase/Private/AnimDatabase.cpp
--- a/Engine/Plugins/Experimental/Animation/AnimDatabase/Source/AnimDatabase/Private/AnimDatabase.cpp
+++ b/Engine/Plugins/Experimental/Animation/AnimDatabase/Source/AnimDatabase/Private/AnimDatabase.cpp
@@ -2296,11 +2296,11 @@ namespace UE::AnimDatabase::Private
 
 void UAnimDatabaseQuery::UpdateQueryEntries()
 {
-	QueryEntries.Reset();
-
 	if (Database)
 	{
 		check(QueryRanges.IsValid());
+		QueryEntries.SetNum(QueryRanges.FrameRangeSet->GetTotalRangeNum());
+		int32 QueryEntryIdx = 0;
 
 		const int32 EntryNum = QueryRanges.FrameRangeSet->GetEntryNum();
 
@@ -2314,18 +2314,31 @@ void UAnimDatabaseQuery::UpdateQueryEntries()
 
 			for (int32 EntryRangeIdx = 0; EntryRangeIdx < RangeNum; EntryRangeIdx++)
 			{
-				TSharedPtr<UE::AnimDatabase::Editor::FQueryEntry> QueryEntry = MakeShared<UE::AnimDatabase::Editor::FQueryEntry>();
-				QueryEntry->Sequence = AnimSequence;
-				QueryEntry->StartFrame = QueryRanges.FrameRangeSet->GetEntryRangeStart(EntryIdx, EntryRangeIdx);
-				QueryEntry->StopFrame = QueryEntry->StartFrame + QueryRanges.FrameRangeSet->GetEntryRangeLength(EntryIdx, EntryRangeIdx) - 1;
-				QueryEntry->bIsMirrored = bIsMirrored;
-				QueryEntry->bIsSelected = false;
+				TSharedPtr<UE::AnimDatabase::Editor::FQueryEntry>& QueryEntry = QueryEntries[QueryEntryIdx];
+				if (!QueryEntry.IsValid()) { QueryEntry = MakeShared<UE::AnimDatabase::Editor::FQueryEntry>(); }
+
+				const int32 NewStartFrame = QueryRanges.FrameRangeSet->GetEntryRangeStart(EntryIdx, EntryRangeIdx);
+				const int32 NewStopFrame = QueryRanges.FrameRangeSet->GetEntryRangeStop(EntryIdx, EntryRangeIdx) - 1;
+				
+				// If the entry has changed then reset selection
+
+				if (QueryEntry->Sequence != AnimSequence ||
+					QueryEntry->StartFrame != NewStartFrame ||
+					QueryEntry->StopFrame != NewStopFrame ||
+					QueryEntry->bIsMirrored != bIsMirrored)
+				{
+					QueryEntry->Sequence = AnimSequence;
+					QueryEntry->StartFrame = NewStartFrame;
+					QueryEntry->StopFrame = NewStopFrame;
+					QueryEntry->bIsMirrored = bIsMirrored;
+					QueryEntry->bIsSelected = false;
+				}
 
 				// Here we use a hash of the sequence, start, and stop frames to produce the color
 
 				const uint32 HashData[5] =
 				{
-					(uint32)RangeIndentifierColorSeed,
+					(uint32)RangeIdentifierColorSeed,
 					GetTypeHash(AnimSequence),
 					(uint32)QueryEntry->StartFrame,
 					(uint32)QueryEntry->StopFrame,
@@ -2339,28 +2352,21 @@ void UAnimDatabaseQuery::UpdateQueryEntries()
 
 				QueryEntry->Color = FLinearColor::MakeFromHSV8(Hue, Saturation, Brightness);
 
-				QueryEntries.Add(QueryEntry);
+				QueryEntryIdx++;
 			}
 		}
+
+		check(QueryEntryIdx == QueryEntries.Num());
+	}
+	else
+	{
+		QueryEntries.Empty();
 	}
 }
 
-void UAnimDatabaseQuery::ResetTransient()
+void UAnimDatabaseQuery::ForceRefresh()
 {
 	DatabaseContentHash = 0;
-	QueryRanges = FAnimDatabaseFrameRanges();
-	ActiveFrameRanges = nullptr;
-	QueryEntries.Reset();
-	ActiveRanges.Reset();
-	SelectedRanges.Reset();
-	SelectedFrameRanges = FAnimDatabaseFrameRanges();
-	AdditionalTimelineFramesObjects.Reset();
-	AdditionalTimelineFrameRangesObjects.Reset();
-	AdditionalTimelineFrameAttributeObjects.Reset();
-	ActiveAdditionalTimelineFrames.Reset();
-	ActiveAdditionalTimelineFrameRanges.Reset();
-	ActiveAdditionalTimelineFrameAttributes.Reset();
-	ActiveDebugDrawer = nullptr;
 }
 
 void UAnimDatabaseQuery::DeleteSelectedFromDatabase()
@@ -2417,8 +2423,6 @@ bool UAnimDatabaseQuery::Update()
 			check(QueryRanges.IsValid());
 
 			UpdateQueryEntries();
-			SelectedRanges.Empty();
-			SelectedFrameRanges = UAnimDatabaseFrameRangesLibrary::MakeEmptyFrameRanges();
 		}
 
 		check(QueryRanges.IsValid());
@@ -2457,33 +2461,13 @@ bool UAnimDatabaseQuery::Update()
 			SelectedFrameRanges = UAnimDatabaseFrameRangesLibrary::FrameRangesGatherRangesFromIndexSet(QueryRanges, SelectedRanges);
 		}
 
-		// Return if query update required
+		// Update Additional Frames/FrameRanges/FrameAttributes
 
-		return bQueryUpdateRequired;
-	}
-	else
-	{
-		if (DatabaseContentHash != 0)
-		{
-			ResetTransient();
-			return true;
-		}
-		else
-		{
-			return false;
-		}
-	}
-}
-
-bool UAnimDatabaseQuery::UpdateAdditionalTimelineObjects()
-{
-	if (Database)
-	{
 		const int32 FramesNum = AdditionalTimelineFrames.Num();
 		const int32 FrameRangesNum = AdditionalTimelineFrameRanges.Num();
 		const int32 FrameAttributeNum = AdditionalTimelineFrameAttributes.Num();
 
-		bool bUpdateRequired =
+		bQueryUpdateRequired = bQueryUpdateRequired ||
 			ActiveAdditionalTimelineFrames.Num() != FramesNum ||
 			ActiveAdditionalTimelineFrameRanges.Num() != FrameRangesNum ||
 			ActiveAdditionalTimelineFrameAttributes.Num() != FrameAttributeNum;
@@ -2498,9 +2482,9 @@ bool UAnimDatabaseQuery::UpdateAdditionalTimelineObjects()
 
 		for (int32 FramesIdx = 0; FramesIdx < FramesNum; FramesIdx++)
 		{
-			if (ActiveAdditionalTimelineFrames[FramesIdx] != AdditionalTimelineFrames[FramesIdx])
+			if (bQueryUpdateRequired || ActiveAdditionalTimelineFrames[FramesIdx] != AdditionalTimelineFrames[FramesIdx])
 			{
-				bUpdateRequired = true;
+				bQueryUpdateRequired = true;
 				ActiveAdditionalTimelineFrames[FramesIdx] = AdditionalTimelineFrames[FramesIdx];
 				AdditionalTimelineFramesObjects[FramesIdx] = UAnimDatabaseFramesLibrary::MakeFramesFromFunction(Database, QueryRanges, AdditionalTimelineFrames[FramesIdx]);
 			}
@@ -2508,9 +2492,9 @@ bool UAnimDatabaseQuery::UpdateAdditionalTimelineObjects()
 
 		for (int32 FrameRangesIdx = 0; FrameRangesIdx < FrameRangesNum; FrameRangesIdx++)
 		{
-			if (ActiveAdditionalTimelineFrameRanges[FrameRangesIdx] != AdditionalTimelineFrameRanges[FrameRangesIdx])
+			if (bQueryUpdateRequired || ActiveAdditionalTimelineFrameRanges[FrameRangesIdx] != AdditionalTimelineFrameRanges[FrameRangesIdx])
 			{
-				bUpdateRequired = true;
+				bQueryUpdateRequired = true;
 				ActiveAdditionalTimelineFrameRanges[FrameRangesIdx] = AdditionalTimelineFrameRanges[FrameRangesIdx];
 				AdditionalTimelineFrameRangesObjects[FrameRangesIdx] = UAnimDatabaseFrameRangesLibrary::MakeFrameRangesFromFunction(Database, QueryRanges, AdditionalTimelineFrameRanges[FrameRangesIdx]);
 			}
@@ -2518,20 +2502,20 @@ bool UAnimDatabaseQuery::UpdateAdditionalTimelineObjects()
 
 		for (int32 FrameAttributeIdx = 0; FrameAttributeIdx < FrameAttributeNum; FrameAttributeIdx++)
 		{
-			if (ActiveAdditionalTimelineFrameAttributes[FrameAttributeIdx] != AdditionalTimelineFrameAttributes[FrameAttributeIdx])
+			if (bQueryUpdateRequired || ActiveAdditionalTimelineFrameAttributes[FrameAttributeIdx] != AdditionalTimelineFrameAttributes[FrameAttributeIdx])
 			{
-				bUpdateRequired = true;
+				bQueryUpdateRequired = true;
 				ActiveAdditionalTimelineFrameAttributes[FrameAttributeIdx] = AdditionalTimelineFrameAttributes[FrameAttributeIdx];
 				AdditionalTimelineFrameAttributeObjects[FrameAttributeIdx] = UAnimDatabaseFrameAttributeLibrary::MakeFrameAttributeFromFunction(Database, QueryRanges, AdditionalTimelineFrameAttributes[FrameAttributeIdx]);
 			}
 		}
 
-		return bUpdateRequired;
-	}
-	else
-	{
-		return false;
+		// Return if query update required
+
+		return bQueryUpdateRequired;
 	}
+	
+	return false;
 }
 
 void UAnimDatabaseQuery::RunFunction()
@@ -7399,6 +7383,64 @@ void UAnimDatabaseFunction_ProcessContacts::Run_Implementation(UAnimDatabase* In
 	}
 }
 
+void UAnimDatabaseFunction_ProcessPhaseCurves::Run_Implementation(UAnimDatabase* InDatabase, const FAnimDatabaseFrameRanges& InFrameRanges, const bool bShouldTransact)
+{
+	const FAnimDatabaseFrameAttribute PhaseAngle = UAnimDatabaseFrameAttributeLibrary::FrameAttributeExtractPhase(
+		InFrameRanges, 
+		UAnimDatabaseFramesLibrary::MakeFramesFromFunction(InDatabase, InFrameRanges, ZeroPhaseFrames),
+		UAnimDatabaseFramesLibrary::MakeFramesFromFunction(InDatabase, InFrameRanges, HalfPhaseFrames),
+		ExtrapolationMode);
+
+	if (bAddPhaseAngle)
+	{
+		FAnimDatabaseFrameAttribute PhaseAngleValue = UAnimDatabaseFrameAttributeLibrary::FrameAttributeAngleToFloatDegrees(PhaseAngle);
+
+		if (bRescalePhaseAngle)
+		{
+			PhaseAngleValue = UAnimDatabaseFrameAttributeLibrary::FrameAttributeAddFloat(UAnimDatabaseFrameAttributeLibrary::FrameAttributeDivideFloat(PhaseAngleValue, 360.0f), 0.5f);
+		}
+
+		UAnimDatabaseLibrary::DatabaseAddCurves(
+			InDatabase,
+			{ PhaseAngleCurveName },
+			{ PhaseAngleValue },
+			{ PhaseAngleCurveColor },
+			bSparseKeys,
+			bShouldTransact);
+	}
+
+	if (bAddPhaseDirection)
+	{
+		const FAnimDatabaseFrameAttribute PhaseDirectionX = UAnimDatabaseFrameAttributeLibrary::FrameAttributeCos(PhaseAngle);
+		const FAnimDatabaseFrameAttribute PhaseDirectionY = UAnimDatabaseFrameAttributeLibrary::FrameAttributeSin(PhaseAngle);
+
+		UAnimDatabaseLibrary::DatabaseAddCurves(
+			InDatabase,
+			{ PhaseDirectionXCurveName, PhaseDirectionYCurveName },
+			{ PhaseDirectionX, PhaseDirectionY },
+			{ PhaseDirectionXCurveColor, PhaseDirectionYCurveColor },
+			bSparseKeys,
+			bShouldTransact);
+	}
+
+	if (bAddPhaseRamp)
+	{
+		const FAnimDatabaseFrameAttribute PhaseRamp = UAnimDatabaseFrameAttributeLibrary::FrameAttributeSubtractFloat(
+			UAnimDatabaseFrameAttributeLibrary::FrameAttributeMultiplyFloat(
+				UAnimDatabaseFrameAttributeLibrary::FrameAttributeAbs(
+					UAnimDatabaseFrameAttributeLibrary::FrameAttributeDivideFloat(
+						UAnimDatabaseFrameAttributeLibrary::FrameAttributeAngleToFloatDegrees(PhaseAngle), 180.0f)), 2.0f), 1.0f);
+
+		UAnimDatabaseLibrary::DatabaseAddCurves(
+			InDatabase,
+			{ PhaseRampCurveName },
+			{ PhaseRamp },
+			{ PhaseRampCurveColor },
+			bSparseKeys,
+			bShouldTransact);
+	}
+}
+
 void UAnimDatabaseFunction_SetBoneLocalTransformFromFrameAttribute::Run_Implementation(UAnimDatabase* InDatabase, const FAnimDatabaseFrameRanges& InFrameRanges, const bool bShouldTransact)
 {
 	UAnimDatabaseLibrary::DatabaseSetLocalBoneTransform(InDatabase, BoneName, UAnimDatabaseFrameAttributeLibrary::MakeFrameAttributeFromFunction(InDatabase, InFrameRanges, FrameAttribute), bShouldTransact);
diff --git a/Engine/Plugins/Experimental/Animation/AnimDatabase/Source/AnimDatabase/Private/AnimDatabaseFrameAttribute.cpp b/Engine/Plugins/Experimental/Animation/AnimDatabase/Source/AnimDatabase/Private/AnimDatabaseFrameAttribute.cpp
--- a/Engine/Plugins/Experimental/Animation/AnimDatabase/Source/AnimDatabase/Private/AnimDatabaseFrameAttribute.cpp
+++ b/Engine/Plugins/Experimental/Animation/AnimDatabase/Source/AnimDatabase/Private/AnimDatabaseFrameAttribute.cpp
@@ -8286,6 +8286,66 @@ FAnimDatabaseFrameAttribute UAnimDatabaseFrameAttributeLibrary::FrameAttributeEx
 	return Out;
 }
 
+FAnimDatabaseFrameAttribute UAnimDatabaseFrameAttributeLibrary::FrameAttributeSin(const FAnimDatabaseFrameAttribute& A)
+{
+	if (!A.IsValid())
+	{
+		UE_LOGFMT(LogAnimDatabase, Error, "FrameAttributeSin: Invalid FrameAttribute.");
+		return FAnimDatabaseFrameAttribute();
+	}
+
+	if (A.Type != EAnimDatabaseAttributeType::Angle)
+	{
+		UE_LOGFMT(LogAnimDatabase, Error, "FrameAttributeSin: FrameAttribute Types invalid. Got {Type}.", AttributeTypeNameInternal(A.Type));
+		return FAnimDatabaseFrameAttribute();
+	}
+
+	FAnimDatabaseFrameAttribute Out = MakeEmptyFrameAttribute();
+	Out.Type = EAnimDatabaseAttributeType::Float;
+	UE::Learning::FrameAttribute::Sin(*Out.FrameAttribute, *A.FrameAttribute);
+	return Out;
+}
+
+FAnimDatabaseFrameAttribute UAnimDatabaseFrameAttributeLibrary::FrameAttributeCos(const FAnimDatabaseFrameAttribute& A)
+{
+	if (!A.IsValid())
+	{
+		UE_LOGFMT(LogAnimDatabase, Error, "FrameAttributeCos: Invalid FrameAttribute.");
+		return FAnimDatabaseFrameAttribute();
+	}
+
+	if (A.Type != EAnimDatabaseAttributeType::Angle)
+	{
+		UE_LOGFMT(LogAnimDatabase, Error, "FrameAttributeCos: FrameAttribute Types invalid. Got {Type}.", AttributeTypeNameInternal(A.Type));
+		return FAnimDatabaseFrameAttribute();
+	}
+
+	FAnimDatabaseFrameAttribute Out = MakeEmptyFrameAttribute();
+	Out.Type = EAnimDatabaseAttributeType::Float;
+	UE::Learning::FrameAttribute::Cos(*Out.FrameAttribute, *A.FrameAttribute);
+	return Out;
+}
+
+FAnimDatabaseFrameAttribute UAnimDatabaseFrameAttributeLibrary::FrameAttributeAtan2(const FAnimDatabaseFrameAttribute& Y, const FAnimDatabaseFrameAttribute& X)
+{
+	if (!Y.IsValid() || !X.IsValid())
+	{
+		UE_LOGFMT(LogAnimDatabase, Error, "FrameAttributeAtan2: Invalid FrameAttribute.");
+		return FAnimDatabaseFrameAttribute();
+	}
+
+	if (Y.Type != EAnimDatabaseAttributeType::Float || X.Type != EAnimDatabaseAttributeType::Float)
+	{
+		UE_LOGFMT(LogAnimDatabase, Error, "FrameAttributeAtan2: FrameAttribute Types invalid. Got {Type} and {Type}.", AttributeTypeNameInternal(Y.Type), AttributeTypeNameInternal(X.Type));
+		return FAnimDatabaseFrameAttribute();
+	}
+
+	FAnimDatabaseFrameAttribute Out = MakeEmptyFrameAttribute();
+	Out.Type = EAnimDatabaseAttributeType::Angle;
+	UE::Learning::FrameAttribute::Atan2(*Out.FrameAttribute, *Y.FrameAttribute, *X.FrameAttribute);
+	return Out;
+}
+
 FAnimDatabaseFrameAttribute UAnimDatabaseFrameAttributeLibrary::FrameAttributeLength(const FAnimDatabaseFrameAttribute& Vector)
 {
 	TRACE_CPUPROFILER_EVENT_SCOPE(UAnimDatabaseFrameAttributeLibrary::FrameAttributeLength);
@@ -8925,6 +8985,22 @@ FAnimDatabaseFrameAttribute UAnimDatabaseFrameAttributeLibrary::FrameAttributeVe
 	return Out;
 }
 
+FAnimDatabaseFrameAttribute UAnimDatabaseFrameAttributeLibrary::FrameAttributeExtractPhase(const FAnimDatabaseFrameRanges& FrameRanges, const FAnimDatabaseFrames& ZeroPhaseFrames, const FAnimDatabaseFrames& HalfPhaseFrames, const EAnimDatabasePhaseExtrapolationMode ExtrapolationMode)
+{
+	if (!FrameRanges.IsValid() || !ZeroPhaseFrames.IsValid() || !HalfPhaseFrames.IsValid())
+	{
+		UE_LOGFMT(LogAnimDatabase, Error, "FrameAttributeExtractPhase: Invalid FrameRanges or Frames.");
+		return FAnimDatabaseFrameAttribute();
+	}
+
+	FAnimDatabaseFrameAttribute Out = MakeEmptyFrameAttribute();
+	Out.Type = EAnimDatabaseAttributeType::Angle;
+	UE::Learning::FrameAttribute::ExtractPhase(*Out.FrameAttribute, *FrameRanges.FrameRangeSet, *ZeroPhaseFrames.FrameSet, *HalfPhaseFrames.FrameSet,
+		ExtrapolationMode == EAnimDatabasePhaseExtrapolationMode::Extrapolate ? UE::Learning::FrameAttribute::EPhaseExtrapolationMode::Extrapolate :
+		ExtrapolationMode == EAnimDatabasePhaseExtrapolationMode::Repeat ? UE::Learning::FrameAttribute::EPhaseExtrapolationMode::Repeat : UE::Learning::FrameAttribute::EPhaseExtrapolationMode::Repeat);
+	return Out;
+}
+
 FAnimDatabaseFrameAttribute UAnimDatabaseFrameAttributeLibrary::FrameAttributeFilterGaussian(const FAnimDatabaseFrameAttribute& A, const float StandardDeviationInFrames)
 {
 	if (!A.IsValid())
@@ -10306,6 +10382,26 @@ FAnimDatabaseFrameAttribute UAnimDatabaseFrameAttributeFunction_ContactCurve::Ma
 		SmoothingAmount);
 }
 
+FAnimDatabaseFrames UAnimDatabaseFramesFunction_ContactFrames::MakeFrames_Implementation(const UAnimDatabase* Database, const FAnimDatabaseFrameRanges& FrameRanges) const
+{
+	FAnimDatabaseFrames ContactFrames = UAnimDatabaseFrameAttributeLibrary::FrameAttributeContactFrames(
+		Database,
+		FrameRanges,
+		BoneName,
+		HeightThreshold,
+		VelocityThreshold,
+		bFilterCurve,
+		MajorityVoteFilterWidth,
+		EventOffset);
+
+	if (EventFilter)
+	{
+		ContactFrames = UAnimDatabaseFrameRangesLibrary::FramesFrameRangesIntersection(ContactFrames, UAnimDatabaseFrameRangesLibrary::MakeFrameRangesFromFunction(Database, FrameRanges, EventFilter));
+	}
+
+	return ContactFrames;
+}
+
 FAnimDatabaseFrameAttribute UAnimDatabaseFrameAttributeFunction_ContactEvent::MakeFrameAttribute_Implementation(const UAnimDatabase* Database, const FAnimDatabaseFrameRanges& FrameRanges) const
 {
 	FAnimDatabaseFrames ContactFrames = UAnimDatabaseFrameAttributeLibrary::FrameAttributeContactFrames(
diff --git a/Engine/Plugins/Experimental/Animation/AnimDatabase/Source/AnimDatabase/Private/AnimDatabaseFrameRanges.cpp b/Engine/Plugins/Experimental/Animation/AnimDatabase/Source/AnimDatabase/Private/AnimDatabaseFrameRanges.cpp
--- a/Engine/Plugins/Experimental/Animation/AnimDatabase/Source/AnimDatabase/Private/AnimDatabaseFrameRanges.cpp
+++ b/Engine/Plugins/Experimental/Animation/AnimDatabase/Source/AnimDatabase/Private/AnimDatabaseFrameRanges.cpp
@@ -303,6 +303,96 @@ FAnimDatabaseFrames UAnimDatabaseFramesLibrary::MakeFramesFromAnimNotify(const U
 	return Out;
 }
 
+FAnimDatabaseFrames UAnimDatabaseFramesLibrary::MakeFramesFromAnimNotifiesOnTrack(const UAnimDatabase* Database, const FAnimDatabaseFrameRanges& FrameRanges, const FName TrackName)
+{
+	TRACE_CPUPROFILER_EVENT_SCOPE(UAnimDatabaseFramesLibrary::MakeFramesFromAnimNotifiesOnTrack);
+
+	if (!Database)
+	{
+		UE_LOGFMT(LogAnimDatabase, Error, "MakeFramesFromAnimNotifiesOnTrack: Database is nullptr.");
+		return FAnimDatabaseFrames();
+	}
+
+	if (!FrameRanges.IsValid())
+	{
+		UE_LOGFMT(LogAnimDatabase, Error, "MakeFramesFromAnimNotifiesOnTrack: Invalid FrameRanges Object.");
+		return FAnimDatabaseFrames();
+	}
+
+	FAnimDatabaseFrames Out = MakeEmptyFrames();
+
+	TArray<FAnimNotifyEventReference, TInlineAllocator<64>> SortedReferences;
+	TArray<int32> FramesToAdd;
+
+	const int32 EntryNum = FrameRanges.FrameRangeSet->GetEntryNum();
+
+	for (int32 EntryIdx = 0; EntryIdx < EntryNum; EntryIdx++)
+	{
+		const int32 SequenceIdx = FrameRanges.FrameRangeSet->GetEntrySequence(EntryIdx);
+
+		const UAnimSequence* Sequence = Database->GetAnimSequence(SequenceIdx);
+		if (!Sequence) { continue; }
+
+		int32 TrackIndex = INDEX_NONE;
+
+#if WITH_EDITOR
+		int32 TrackNum = Sequence->AnimNotifyTracks.Num();
+		for (int32 TrackIdx = 0; TrackIdx < TrackNum; TrackIdx++)
+		{
+			if (Sequence->AnimNotifyTracks[TrackIdx].TrackName == TrackName)
+			{
+				TrackIndex = TrackIdx;
+				break;
+			}
+		}
+#else
+		UE_LOGFMT(LogAnimDatabase, Error, "MakeFramesFromAnimNotifiesOnTrack: Track Names are only available in Editor.");
+		return FAnimDatabaseFrames();
+#endif
+
+		if (TrackIndex == INDEX_NONE) { continue; }
+
+		const int32 SequenceFrameNum = FMath::Max(Database->GetSequenceFrameNum(SequenceIdx), 1);
+
+		FAnimNotifyContext Notifies;
+		Sequence->GetAnimNotifies(0.0f, Sequence->GetPlayLength(), Notifies);
+
+		SortedReferences = Notifies.ActiveNotifies;
+		SortedReferences.Sort([](const FAnimNotifyEventReference& A, const FAnimNotifyEventReference& B) { return A.GetNotify()->GetTime() < B.GetNotify()->GetTime(); });
+
+		int32 PreviousFramesTime = -1;
+
+		FramesToAdd.Reset();
+		for (const FAnimNotifyEventReference& AnimNotifyRef : SortedReferences)
+		{
+			if (AnimNotifyRef.GetNotify()->Notify && AnimNotifyRef.GetNotify()->TrackIndex == TrackIndex)
+			{
+				const int32 AnimNotifyStart = FMath::Clamp(
+					FMath::RoundToInt(AnimNotifyRef.GetNotify()->GetTime() * Database->GetFrameRate().AsDecimal()),
+					0, SequenceFrameNum - 1);
+
+				if (AnimNotifyStart == PreviousFramesTime)
+				{
+					UE_LOGFMT(LogAnimDatabase, Warning, "MakeFramesFromAnimNotifiesOnTrack: Anim Notify {Notify} is at same time as Previous in Anim Sequence {Name} Time {Time} Frame {Frame}", *AnimNotifyRef.GetNotify()->Notify->GetName(), *Sequence->GetName(), AnimNotifyRef.GetNotify()->GetTime(), AnimNotifyStart);
+				}
+				else
+				{
+					if (FrameRanges.FrameRangeSet->EntryContains(EntryIdx, AnimNotifyStart))
+					{
+						FramesToAdd.Add(AnimNotifyStart);
+					}
+
+					PreviousFramesTime = AnimNotifyStart;
+				}
+			}
+		}
+
+		Out.FrameSet->AddEntry(SequenceIdx, FramesToAdd);
+	}
+
+	return Out;
+}
+
 FAnimDatabaseFrames UAnimDatabaseFramesLibrary::MakeFramesFromSyncMarker(const UAnimDatabase* Database, const FAnimDatabaseFrameRanges& FrameRanges, const FName SyncMarkerName)
 {
 	TRACE_CPUPROFILER_EVENT_SCOPE(UAnimDatabaseFramesLibrary::MakeFramesFromSyncMarker);
@@ -875,6 +965,26 @@ FAnimDatabaseFrames UAnimDatabaseFramesFunction_AnimNotify::MakeFrames_Implement
 	return UAnimDatabaseFramesLibrary::MakeFramesFromAnimNotify(Database, FrameRanges, AnimNotify);
 }
 
+FAnimDatabaseFrames UAnimDatabaseFramesFunction_AnimNotifiesOnTrack::MakeFrames_Implementation(const UAnimDatabase* Database, const FAnimDatabaseFrameRanges& FrameRanges) const
+{
+	return UAnimDatabaseFramesLibrary::MakeFramesFromAnimNotifiesOnTrack(Database, FrameRanges, TrackName);
+}
+
+FAnimDatabaseFrames UAnimDatabaseFramesFunction_SyncMarker::MakeFrames_Implementation(const UAnimDatabase* Database, const FAnimDatabaseFrameRanges& FrameRanges) const
+{
+	return UAnimDatabaseFramesLibrary::MakeFramesFromSyncMarker(Database, FrameRanges, SyncMarker);
+}
+
+FAnimDatabaseFrames UAnimDatabaseFramesFunction_FrameRangeStarts::MakeFrames_Implementation(const UAnimDatabase* Database, const FAnimDatabaseFrameRanges& FrameRanges) const
+{
+	return UAnimDatabaseFrameRangesLibrary::MakeFramesAtFrameRangesStarts(UAnimDatabaseFrameRangesLibrary::MakeFrameRangesFromFunction(Database, FrameRanges, Function));
+}
+
+FAnimDatabaseFrames UAnimDatabaseFramesFunction_FrameRangeEnds::MakeFrames_Implementation(const UAnimDatabase* Database, const FAnimDatabaseFrameRanges& FrameRanges) const
+{
+	return UAnimDatabaseFrameRangesLibrary::MakeFramesAtFrameRangesEnds(UAnimDatabaseFrameRangesLibrary::MakeFrameRangesFromFunction(Database, FrameRanges, Function));
+}
+
 FAnimDatabaseFrameRanges UAnimDatabaseFrameRangesLibrary::MakeEmptyFrameRanges()
 {
 	FAnimDatabaseFrameRanges Out;
@@ -2145,6 +2255,63 @@ FAnimDatabaseFrameRanges UAnimDatabaseFrameRangesLibrary::FrameRangesNot(const U
 	return FrameRangesDifference(MakeFrameRangesFromDatabase(Database), FrameRanges);
 }
 
+FAnimDatabaseFrameRanges UAnimDatabaseFrameRangesLibrary::FrameRangesBounds(const FAnimDatabaseFrameRanges& FrameRanges)
+{
+	if (!FrameRanges.IsValid())
+	{
+		UE_LOGFMT(LogAnimDatabase, Error, "FrameRangesBounds: Invalid FrameRanges.");
+		return FAnimDatabaseFrameRanges();
+	}
+
+	FAnimDatabaseFrameRanges Out = MakeEmptyFrameRanges();
+	UE::Learning::FrameRangeSet::Bounds(*Out.FrameRangeSet, *FrameRanges.FrameRangeSet);
+	Out.FrameRangeSet->Check();
+	return Out;
+}
+
+
+FAnimDatabaseFrameRanges UAnimDatabaseFrameRangesLibrary::FrameRangesBoundsFrames(const FAnimDatabaseFrames& Frames)
+{
+	if (!Frames.IsValid())
+	{
+		UE_LOGFMT(LogAnimDatabase, Error, "FrameRangesBoundsFrames: Invalid Frames.");
+		return FAnimDatabaseFrameRanges();
+	}
+
+	FAnimDatabaseFrameRanges Out = MakeEmptyFrameRanges();
+	UE::Learning::FrameRangeSet::Bounds(*Out.FrameRangeSet, *Frames.FrameSet);
+	Out.FrameRangeSet->Check();
+	return Out;
+}
+
+FAnimDatabaseFrameRanges UAnimDatabaseFrameRangesLibrary::FrameRangesIntersects(const FAnimDatabaseFrameRanges& A, const FAnimDatabaseFrameRanges& B)
+{
+	if (!A.IsValid() || !B.IsValid())
+	{
+		UE_LOGFMT(LogAnimDatabase, Error, "FrameRangesIntersects: Invalid Frame Ranges.");
+		return FAnimDatabaseFrameRanges();
+	}
+
+	FAnimDatabaseFrameRanges Out = MakeEmptyFrameRanges();
+	UE::Learning::FrameRangeSet::Intersects(*Out.FrameRangeSet, *A.FrameRangeSet, *B.FrameRangeSet);
+	Out.FrameRangeSet->Check();
+	return Out;
+}
+
+FAnimDatabaseFrameRanges UAnimDatabaseFrameRangesLibrary::FrameRangesIntersectsFrames(const FAnimDatabaseFrameRanges& A, const FAnimDatabaseFrames& B)
+{
+	if (!A.IsValid() || !B.IsValid())
+	{
+		UE_LOGFMT(LogAnimDatabase, Error, "FrameRangesIntersectsFrames: Invalid Frame Ranges or Frames.");
+		return FAnimDatabaseFrameRanges();
+	}
+
+	FAnimDatabaseFrameRanges Out = MakeEmptyFrameRanges();
+	UE::Learning::FrameRangeSet::Intersects(*Out.FrameRangeSet, *A.FrameRangeSet, *B.FrameSet);
+	Out.FrameRangeSet->Check();
+	return Out;
+}
+
 FAnimDatabaseFrameRanges UAnimDatabaseFrameRangesLibrary::FrameRangesTrim(const FAnimDatabaseFrameRanges& A, const int32 TrimStartFrames, const int32 TrimEndFrames)
 {
 	if (!A.IsValid())
@@ -2154,7 +2321,7 @@ FAnimDatabaseFrameRanges UAnimDatabaseFrameRangesLibrary::FrameRangesTrim(const
 	}
 
 	FAnimDatabaseFrameRanges Out = MakeEmptyFrameRanges();
-	UE::Learning::FrameRangeSet::Trim(*Out.FrameRangeSet, *A.FrameRangeSet, TrimStartFrames, TrimEndFrames);
+	UE::Learning::FrameRangeSet::Trim(*Out.FrameRangeSet, *A.FrameRangeSet, FMath::Max(TrimStartFrames, 0), FMath::Max(TrimEndFrames, 0));
 	Out.FrameRangeSet->Check();
 	return Out;
 }
@@ -2168,7 +2335,7 @@ FAnimDatabaseFrameRanges UAnimDatabaseFrameRangesLibrary::FrameRangesTrimStart(c
 	}
 
 	FAnimDatabaseFrameRanges Out = MakeEmptyFrameRanges();
-	UE::Learning::FrameRangeSet::TrimStart(*Out.FrameRangeSet, *A.FrameRangeSet, TrimFrames);
+	UE::Learning::FrameRangeSet::TrimStart(*Out.FrameRangeSet, *A.FrameRangeSet, FMath::Max(TrimFrames, 0));
 	Out.FrameRangeSet->Check();
 	return Out;
 }
@@ -2182,7 +2349,7 @@ FAnimDatabaseFrameRanges UAnimDatabaseFrameRangesLibrary::FrameRangesTrimEnd(con
 	}
 
 	FAnimDatabaseFrameRanges Out = MakeEmptyFrameRanges();
-	UE::Learning::FrameRangeSet::TrimEnd(*Out.FrameRangeSet, *A.FrameRangeSet, TrimFrames);
+	UE::Learning::FrameRangeSet::TrimEnd(*Out.FrameRangeSet, *A.FrameRangeSet, FMath::Max(TrimFrames, 0));
 	Out.FrameRangeSet->Check();
 	return Out;
 }
@@ -2196,7 +2363,7 @@ FAnimDatabaseFrameRanges UAnimDatabaseFrameRangesLibrary::FrameRangesPad(const F
 	}
 
 	FAnimDatabaseFrameRanges Out = MakeEmptyFrameRanges();
-	UE::Learning::FrameRangeSet::Pad(*Out.FrameRangeSet, *A.FrameRangeSet, PadStartFrames, PadEndFrames);
+	UE::Learning::FrameRangeSet::Pad(*Out.FrameRangeSet, *A.FrameRangeSet, FMath::Max(PadStartFrames, 0), FMath::Max(PadEndFrames, 0));
 	Out.FrameRangeSet->Check();
 	return Out;
 }
@@ -2210,7 +2377,7 @@ FAnimDatabaseFrameRanges UAnimDatabaseFrameRangesLibrary::FrameRangesPadStart(co
 	}
 
 	FAnimDatabaseFrameRanges Out = MakeEmptyFrameRanges();
-	UE::Learning::FrameRangeSet::PadStart(*Out.FrameRangeSet, *A.FrameRangeSet, PadFrames);
+	UE::Learning::FrameRangeSet::PadStart(*Out.FrameRangeSet, *A.FrameRangeSet, FMath::Max(PadFrames, 0));
 	return Out;
 }
 
@@ -2223,7 +2390,7 @@ FAnimDatabaseFrameRanges UAnimDatabaseFrameRangesLibrary::FrameRangesPadEnd(cons
 	}
 
 	FAnimDatabaseFrameRanges Out = MakeEmptyFrameRanges();
-	UE::Learning::FrameRangeSet::PadEnd(*Out.FrameRangeSet, *A.FrameRangeSet, PadFrames);
+	UE::Learning::FrameRangeSet::PadEnd(*Out.FrameRangeSet, *A.FrameRangeSet, FMath::Max(PadFrames, 0));
 	Out.FrameRangeSet->Check();
 	return Out;
 }
@@ -2958,6 +3125,11 @@ FAnimDatabaseFrameRanges UAnimDatabaseFrameRangesFunction_AnimNotify::MakeFrameR
 	return UAnimDatabaseFrameRangesLibrary::MakeFrameRangesFromAnimNotify(Database, FrameRanges, AnimNotify);
 }
 
+FAnimDatabaseFrameRanges UAnimDatabaseFrameRangesFunction_SyncMarker::MakeFrameRanges_Implementation(const UAnimDatabase* Database, const FAnimDatabaseFrameRanges& FrameRanges) const
+{
+	return UAnimDatabaseFrameRangesLibrary::MakeFrameRangesFromSyncMarker(Database, FrameRanges, SyncMarker);
+}
+
 FAnimDatabaseFrameRanges UAnimDatabaseFrameRangesFunction_AnimNotifyState::MakeFrameRanges_Implementation(const UAnimDatabase* Database, const FAnimDatabaseFrameRanges& FrameRanges) const
 {
 	return UAnimDatabaseFrameRangesLibrary::MakeFrameRangesFromAnimNotifyState(Database, FrameRanges, AnimNotifyState);
@@ -3073,6 +3245,26 @@ FAnimDatabaseFrameRanges UAnimDatabaseFrameRangesFunction_Difference::MakeFrameR
 	return UAnimDatabaseFrameRangesLibrary::FrameRangesDifferenceFromArrayView(FrameRangesSet);
 }
 
+FAnimDatabaseFrameRanges UAnimDatabaseFrameRangesFunction_Bounds::MakeFrameRanges_Implementation(const UAnimDatabase* Database, const FAnimDatabaseFrameRanges& FrameRanges) const
+{
+	return UAnimDatabaseFrameRangesLibrary::FrameRangesBounds(UAnimDatabaseFrameRangesLibrary::MakeFrameRangesFromFunction(Database, FrameRanges, Function));
+}
+
+FAnimDatabaseFrameRanges UAnimDatabaseFrameRangesFunction_BoundsFrames::MakeFrameRanges_Implementation(const UAnimDatabase* Database, const FAnimDatabaseFrameRanges& FrameRanges) const
+{
+	return UAnimDatabaseFrameRangesLibrary::FrameRangesBoundsFrames(UAnimDatabaseFramesLibrary::MakeFramesFromFunction(Database, FrameRanges, Function));
+}
+
+FAnimDatabaseFrameRanges UAnimDatabaseFrameRangesFunction_Intersects::MakeFrameRanges_Implementation(const UAnimDatabase* Database, const FAnimDatabaseFrameRanges& FrameRanges) const
+{
+	return UAnimDatabaseFrameRangesLibrary::FrameRangesIntersects(FrameRanges, UAnimDatabaseFrameRangesLibrary::MakeFrameRangesFromFunction(Database, FrameRanges, Function));
+}
+
+FAnimDatabaseFrameRanges UAnimDatabaseFrameRangesFunction_IntersectsFrames::MakeFrameRanges_Implementation(const UAnimDatabase* Database, const FAnimDatabaseFrameRanges& FrameRanges) const
+{
+	return UAnimDatabaseFrameRangesLibrary::FrameRangesIntersectsFrames(FrameRanges, UAnimDatabaseFramesLibrary::MakeFramesFromFunction(Database, FrameRanges, Function));
+}
+
 FAnimDatabaseFrameRanges UAnimDatabaseFrameRangesFunction_Interval::MakeFrameRanges_Implementation(const UAnimDatabase* Database, const FAnimDatabaseFrameRanges& FrameRanges) const
 {
 	return UAnimDatabaseFrameRangesLibrary::MakeFrameRangesFromFrames(UAnimDatabaseFrameRangesLibrary::MakeFramesFromFrameRangesInterval(FrameRanges, Interval));
diff --git a/Engine/Plugins/Experimental/Animation/AnimDatabase/Source/AnimDatabase/Public/AnimDatabase.h b/Engine/Plugins/Experimental/Animation/AnimDatabase/Source/AnimDatabase/Public/AnimDatabase.h
--- a/Engine/Plugins/Experimental/Animation/AnimDatabase/Source/AnimDatabase/Public/AnimDatabase.h
+++ b/Engine/Plugins/Experimental/Animation/AnimDatabase/Source/AnimDatabase/Public/AnimDatabase.h
@@ -9,7 +9,6 @@
 #include "Kismet/BlueprintFunctionLibrary.h"
 #include "Animation/AnimNotifies/AnimNotify.h"
 #include "Animation/AnimNotifies/AnimNotifyState.h"
-#include "Logging/LogMacros.h"
 #include "Logging/StructuredLog.h"
 #include "Modules/ModuleManager.h"
 
@@ -300,9 +299,9 @@ public:
 
 #if WITH_EDITOR
 
-	/** Resets all of the transient properties/members. */
+	/** Forces a refresh of the query object */
 	UFUNCTION(CallInEditor, Category = "Query", DisplayName = "Refresh")
-	UE_API void ResetTransient();
+	UE_API void ForceRefresh();
 
 #endif
 
@@ -339,7 +338,7 @@ public:
 	FAnimDatabaseFrameRanges SelectedFrameRanges;
 
 	/** Seed used for the colors of the range identifiers */
-	int32 RangeIndentifierColorSeed = 1234;
+	int32 RangeIdentifierColorSeed = 1234;
 
 #endif
 
@@ -359,16 +358,10 @@ private:
 	/** Internal function for updating the QueryEntries array from an updated QueryRanges */
 	UE_API void UpdateQueryEntries();
 
-public:
-
-	/**
-	 * Checks if the given AdditionalTimelineFrames and AdditionalTimelineFrameRanges have been changed and if so updates the internal cached frames
-	 * and frame ranges objects as well as class default objects. Returns true if these were update and UI corresponding model may need updating
-	 */
-	UE_API bool UpdateAdditionalTimelineObjects();
-
 #endif
 
+public:
+
 #if WITH_EDITORONLY_DATA
 
 	/** Additional Frames to display in the timeline */
@@ -882,11 +875,11 @@ class UAnimDatabaseLibrary : public UBlueprintFunctionLibrary
 
 public:
 
-	/** Gets all of the AnimNotify classes present in a given AnimSequence */
+	/** Gets all the AnimNotify classes present in a given AnimSequence */
 	UFUNCTION(BlueprintCallable, Category = "AnimDatabase")
 	static UE_API void GetAnimSequenceAnimNotifyClasses(TArray<TSubclassOf<UAnimNotify>>& OutClasses, UAnimSequence* AnimSequence);
 
-	/** Gets all of the AnimNotifyState classes present in a given AnimSequence */
+	/** Gets all the AnimNotifyState classes present in a given AnimSequence */
 	UFUNCTION(BlueprintCallable, Category = "AnimDatabase")
 	static UE_API void GetAnimSequenceAnimNotifyStateClasses(TArray<TSubclassOf<UAnimNotifyState>>& OutClasses, UAnimSequence* AnimSequence);
 
@@ -1076,7 +1069,7 @@ public:
 
 private:
 
-	/** Loads into memory all of the blueprint assets of a given class. */
+	/** Loads into memory all the blueprint assets of a given class. */
 	static void LoadAllBlueprintAssetsOfClass(UClass* Class);
 };
 
@@ -1163,7 +1156,7 @@ public:
 #endif
 };
 
-/** Function for exporting all of the AnimNotifies and AnimNotifyStates from the given animations  */
+/** Function for exporting all AnimNotifies and AnimNotifyStates from the given animations  */
 UCLASS(BlueprintType, Blueprintable, meta = (DisplayName = "Export Notifies"))
 class UAnimDatabaseFunction_ExportNotifies : public UAnimDatabaseFunction
 {
@@ -1916,6 +1909,99 @@ public:
 #endif
 };
 
+
+/** Function for generating phase curves from another source */
+UCLASS(BlueprintType, Blueprintable, DontCollapseCategories, meta = (DisplayName = "Process Phase Curves"))
+class UAnimDatabaseFunction_ProcessPhaseCurves : public UAnimDatabaseFunction
+{
+	GENERATED_BODY()
+
+#if WITH_EDITORONLY_DATA
+
+public:
+
+	/** The frames used to label a phase time of zero */
+	UPROPERTY(EditAnywhere, Instanced, Category = "Settings")
+	TObjectPtr<UAnimDatabaseFramesFunction> ZeroPhaseFrames;
+
+	/** The frames used to label that half the phase time has elapsed */
+	UPROPERTY(EditAnywhere, Instanced, Category = "Settings")
+	TObjectPtr<UAnimDatabaseFramesFunction> HalfPhaseFrames;
+
+	/** Phase extrapolation mode for the ends of the animation */
+	UPROPERTY(EditAnywhere, Category = "Settings")
+	EAnimDatabasePhaseExtrapolationMode ExtrapolationMode = EAnimDatabasePhaseExtrapolationMode::Extrapolate;
+
+	/** If to write curves using sparse keys */
+	UPROPERTY(EditAnywhere, Category = "Settings")
+	bool bSparseKeys = true;
+
+public:
+
+	/** If to add the phase angle as a curve */
+	UPROPERTY(EditAnywhere, Category = "Phase Angle")
+	bool bAddPhaseAngle = false;
+
+	/** Name of the angular phase curve */
+	UPROPERTY(EditAnywhere, Category = "Phase Angle", meta = (EditCondition = "bAddPhaseAngle", HideEditConditionToggle))
+	FName PhaseAngleCurveName = TEXT("phase_angle");
+
+	/** Color of the angular phase curve */
+	UPROPERTY(EditAnywhere, Category = "Phase Angle", meta = (EditCondition = "bAddPhaseAngle", HideEditConditionToggle))
+	FLinearColor PhaseAngleCurveColor = FLinearColor(0.5f, 0.25, 0.5f, 1.0f);
+
+	/** If to re-scale the phase angle to be between 0 and 1 */
+	UPROPERTY(EditAnywhere, Category = "Phase Angle", meta = (EditCondition = "bAddPhaseAngle", HideEditConditionToggle))
+	bool bRescalePhaseAngle = false;
+
+public:
+
+	/** If to add two phase direction as two curves */
+	UPROPERTY(EditAnywhere, Category = "Phase Direction")
+	bool bAddPhaseDirection = false;
+
+	/** Name of the x phase curve */
+	UPROPERTY(EditAnywhere, Category = "Phase Direction", meta = (EditCondition = "bAddPhaseDirection", HideEditConditionToggle))
+	FName PhaseDirectionXCurveName = TEXT("phase_x");
+
+	/** Name of the y phase curve */
+	UPROPERTY(EditAnywhere, Category = "Phase Direction", meta = (EditCondition = "bAddPhaseDirection", HideEditConditionToggle))
+	FName PhaseDirectionYCurveName = TEXT("phase_y");
+
+	/** Color of the x phase curve */
+	UPROPERTY(EditAnywhere, Category = "Phase Direction", meta = (EditCondition = "bAddPhaseDirection", HideEditConditionToggle))
+	FLinearColor PhaseDirectionXCurveColor = FLinearColor(0.5f, 0.5f, 0.25, 1.0f);
+
+	/** Color of the y phase curve */
+	UPROPERTY(EditAnywhere, Category = "Phase Direction", meta = (EditCondition = "bAddPhaseDirection", HideEditConditionToggle))
+	FLinearColor PhaseDirectionYCurveColor = FLinearColor(0.5f, 0.5f, 0.75, 1.0f);
+
+public:
+
+	/** If to add the phase angle as a ramp that goes up and down between -1 and 1 */
+	UPROPERTY(EditAnywhere, Category = "Phase Ramp")
+	bool bAddPhaseRamp = false;
+
+	/** Name of the phase ramp curve */
+	UPROPERTY(EditAnywhere, Category = "Phase Ramp", meta = (EditCondition = "bAddPhaseRamp", HideEditConditionToggle))
+	FName PhaseRampCurveName = TEXT("phase");
+
+	/** Color of the phase ramp curve */
+	UPROPERTY(EditAnywhere, Category = "Phase Ramp", meta = (EditCondition = "bAddPhaseRamp", HideEditConditionToggle))
+	FLinearColor PhaseRampCurveColor = FLinearColor(0.5f, 0.75, 0.5f, 1.0f);
+	
+#endif
+
+#if WITH_EDITOR
+
+public:
+
+	UE_API virtual void Run_Implementation(UAnimDatabase* InDatabase, const FAnimDatabaseFrameRanges& InFrameRanges, const bool bShouldTransact) override;
+
+#endif
+};
+
+
 /** Function for setting a bone local transform from frame attribute function  */
 UCLASS(BlueprintType, Blueprintable, meta = (DisplayName = "Set Bone Local Transform from Frame Attribute"))
 class UAnimDatabaseFunction_SetBoneLocalTransformFromFrameAttribute : public UAnimDatabaseFunction
diff --git a/Engine/Plugins/Experimental/Animation/AnimDatabase/Source/AnimDatabase/Public/AnimDatabaseFrameAttribute.h b/Engine/Plugins/Experimental/Animation/AnimDatabase/Source/AnimDatabase/Public/AnimDatabaseFrameAttribute.h
--- a/Engine/Plugins/Experimental/Animation/AnimDatabase/Source/AnimDatabase/Public/AnimDatabaseFrameAttribute.h
+++ b/Engine/Plugins/Experimental/Animation/AnimDatabase/Source/AnimDatabase/Public/AnimDatabaseFrameAttribute.h
@@ -41,6 +41,17 @@ enum class EAnimDatabaseAttributeType : uint8
 	Angle = 12,
 };
 
+/* Mode for handling edges in phase extraction */
+UENUM(BlueprintType)
+enum class EAnimDatabasePhaseExtrapolationMode : uint8
+{
+	// Repeats the phase value at the start and end frames 
+	Repeat = 0,
+
+	// Extrapolates the phase value using the nearest pair of frames
+	Extrapolate = 1,
+};
+
 /**
  * Represents a property associated with a set of frame ranges within a UAnimDatabase.
  *
@@ -1470,6 +1481,18 @@ public:
 	UFUNCTION(BlueprintPure, Category = "AnimDatabase")
 	static UE_API FAnimDatabaseFrameAttribute FrameAttributeExp(const FAnimDatabaseFrameAttribute& A);
 
+	/** Takes the sin value of a frame attribute. Assumes input is an Angle attribute. */
+	UFUNCTION(BlueprintPure, Category = "AnimDatabase")
+	static UE_API FAnimDatabaseFrameAttribute FrameAttributeSin(const FAnimDatabaseFrameAttribute& A);
+
+	/** Takes the cos value of a frame attribute. Assumes input is an Angle attribute. */
+	UFUNCTION(BlueprintPure, Category = "AnimDatabase")
+	static UE_API FAnimDatabaseFrameAttribute FrameAttributeCos(const FAnimDatabaseFrameAttribute& A);
+
+	/** Computes the atan2 value of two frame attributes. Assumes inputs are Float attributes. Output is an Angle attribute. */
+	UFUNCTION(BlueprintPure, Category = "AnimDatabase")
+	static UE_API FAnimDatabaseFrameAttribute FrameAttributeAtan2(const FAnimDatabaseFrameAttribute& Y, const FAnimDatabaseFrameAttribute& X);
+
 public:
 
 	/** Computes inv(A) x B for two frame attributes. Assumes types are compatible */
@@ -1689,6 +1712,12 @@ public:
 	UFUNCTION(BlueprintPure, Category = "AnimDatabase")
 	static UE_API FAnimDatabaseFrameAttribute FrameAttributeVelocityToDirection(const FAnimDatabaseFrameAttribute& VelocityFrameAttribute);
 
+public:
+
+	/** Extracts a phase angle from two sets of frames representing the start of the phase and half way through the phase */
+	UFUNCTION(BlueprintPure, Category = "AnimDatabase")
+	static UE_API FAnimDatabaseFrameAttribute FrameAttributeExtractPhase(const FAnimDatabaseFrameRanges& FrameRanges, const FAnimDatabaseFrames& ZeroPhaseFrames, const FAnimDatabaseFrames& HalfPhaseFrames, const EAnimDatabasePhaseExtrapolationMode ExtrapolationMode = EAnimDatabasePhaseExtrapolationMode::Repeat);
+
 public:
 
 	/** Applies a Gaussian smoothing filter to the values in each range of a frame attribute. Assumes type is compatible */
@@ -2103,6 +2132,50 @@ public:
 	UE_API virtual FAnimDatabaseFrameAttribute MakeFrameAttribute_Implementation(const UAnimDatabase* Database, const FAnimDatabaseFrameRanges& FrameRanges) const override;
 };
 
+/** Computes contact frames */
+UCLASS(BlueprintType, Blueprintable, meta = (DisplayName = "Contact Frames"))
+class UAnimDatabaseFramesFunction_ContactFrames : public UAnimDatabaseFramesFunction
+{
+	GENERATED_BODY()
+
+public:
+
+	/** Name of the bone to compute the contact state for */
+	UPROPERTY(EditAnywhere, Category = "Settings")
+	FName BoneName = NAME_None;
+
+	/** Maximum height (in cm) for which the bone can be considered in contact */
+	UPROPERTY(EditAnywhere, Category = "Settings")
+	float HeightThreshold = 20.0f;
+
+	/** Maximum velocity (in cm/s) for which the bone can be considered in contact */
+	UPROPERTY(EditAnywhere, Category = "Settings")
+	float VelocityThreshold = 25.0f;
+
+	/** If to filter the curve */
+	UPROPERTY(EditAnywhere, Category = "Settings")
+	bool bFilterCurve = true;
+
+	/**
+	 * Filter width (in frames) for the majority vote filter used to remove very short contact transitions.
+	 *Contact periods (or non-contact periods) less than half this duration will be removed.
+	 */
+	UPROPERTY(EditAnywhere, Category = "Settings", meta = (EditCondition = "bFilterCurve", HideEditConditionToggle))
+	int32 MajorityVoteFilterWidth = 5;
+
+	/** The amount of time before the contact (in seconds) to have the event */
+	UPROPERTY(EditAnywhere, Category = "Settings", meta = (ForceUnits = "s"))
+	float EventOffset = 0.0f;
+
+	/** If provided, will be used to filter the contact events */
+	UPROPERTY(EditAnywhere, Instanced, Category = "Settings")
+	TObjectPtr<UAnimDatabaseFrameRangesFunction> EventFilter;
+
+public:
+
+	UE_API virtual FAnimDatabaseFrames MakeFrames_Implementation(const UAnimDatabase* Database, const FAnimDatabaseFrameRanges& FrameRanges) const override;
+};
+
 /** Computes a contact event attribute */
 UCLASS(BlueprintType, Blueprintable, meta = (DisplayName = "Contact Event"))
 class UAnimDatabaseFrameAttributeFunction_ContactEvent : public UAnimDatabaseFrameAttributeFunction
diff --git a/Engine/Plugins/Experimental/Animation/AnimDatabase/Source/AnimDatabase/Public/AnimDatabaseFrameRanges.h b/Engine/Plugins/Experimental/Animation/AnimDatabase/Source/AnimDatabase/Public/AnimDatabaseFrameRanges.h
--- a/Engine/Plugins/Experimental/Animation/AnimDatabase/Source/AnimDatabase/Public/AnimDatabaseFrameRanges.h
+++ b/Engine/Plugins/Experimental/Animation/AnimDatabase/Source/AnimDatabase/Public/AnimDatabaseFrameRanges.h
@@ -130,6 +130,10 @@ public:
 	UFUNCTION(BlueprintPure, Category = "AnimDatabase", meta = (NotBlueprintThreadSafe))
 	static UE_API FAnimDatabaseFrames MakeFramesFromAnimNotify(const UAnimDatabase* Database, const FAnimDatabaseFrameRanges& FrameRanges, const TSubclassOf<UAnimNotify> Notify);
 
+	/** Make a set of Frames from the times given by all the instances of AnimNotifies in the database that are on a track with the given name */
+	UFUNCTION(BlueprintPure, Category = "AnimDatabase", meta = (NotBlueprintThreadSafe))
+	static UE_API FAnimDatabaseFrames MakeFramesFromAnimNotifiesOnTrack(const UAnimDatabase* Database, const FAnimDatabaseFrameRanges& FrameRanges, const FName TrackName);
+
 	/** Make a set of Frames from the times given by all the instances of the given SyncMarker in the database */
 	UFUNCTION(BlueprintPure, Category = "AnimDatabase", meta = (NotBlueprintThreadSafe))
 	static UE_API FAnimDatabaseFrames MakeFramesFromSyncMarker(const UAnimDatabase* Database, const FAnimDatabaseFrameRanges& FrameRanges, const FName SyncMarkerName);
@@ -239,7 +243,7 @@ public:
 	UE_API virtual FAnimDatabaseFrames MakeFrames_Implementation(const UAnimDatabase* Database, const FAnimDatabaseFrameRanges& FrameRanges) const override;
 };
 
-/** Finds all the frame for a given Anim Notify Class */
+/** Finds all the frames for a given Anim Notify Class */
 UCLASS(BlueprintType, Blueprintable, meta = (DisplayName = "Anim Notify"))
 class UAnimDatabaseFramesFunction_AnimNotify : public UAnimDatabaseFramesFunction
 {
@@ -255,6 +259,70 @@ public:
 	UE_API virtual FAnimDatabaseFrames MakeFrames_Implementation(const UAnimDatabase* Database, const FAnimDatabaseFrameRanges& FrameRanges) const override;
 };
 
+/** Finds all the frames for all Anim Notifies on a given track */
+UCLASS(BlueprintType, Blueprintable, meta = (DisplayName = "Anim Notifies on Track"))
+class UAnimDatabaseFramesFunction_AnimNotifiesOnTrack : public UAnimDatabaseFramesFunction
+{
+	GENERATED_BODY()
+
+public:
+
+	UPROPERTY(EditAnywhere, Category = "Settings")
+	FName TrackName = NAME_None;
+
+public:
+
+	UE_API virtual FAnimDatabaseFrames MakeFrames_Implementation(const UAnimDatabase* Database, const FAnimDatabaseFrameRanges& FrameRanges) const override;
+};
+
+/** Finds all the frames for a given Sync Marker */
+UCLASS(BlueprintType, Blueprintable, meta = (DisplayName = "Sync Marker"))
+class UAnimDatabaseFramesFunction_SyncMarker : public UAnimDatabaseFramesFunction
+{
+	GENERATED_BODY()
+
+public:
+
+	UPROPERTY(EditAnywhere, Category = "Settings")
+	FName SyncMarker = NAME_None;
+
+public:
+
+	UE_API virtual FAnimDatabaseFrames MakeFrames_Implementation(const UAnimDatabase* Database, const FAnimDatabaseFrameRanges& FrameRanges) const override;
+};
+
+/** Finds all the frames at the start of the given frame ranges */
+UCLASS(BlueprintType, Blueprintable, meta = (DisplayName = "Frame Range Starts"))
+class UAnimDatabaseFramesFunction_FrameRangeStarts : public UAnimDatabaseFramesFunction
+{
+	GENERATED_BODY()
+
+public:
+
+	UPROPERTY(EditAnywhere, Instanced, Category = "Settings")
+	TObjectPtr<UAnimDatabaseFrameRangesFunction> Function;
+
+public:
+
+	UE_API virtual FAnimDatabaseFrames MakeFrames_Implementation(const UAnimDatabase* Database, const FAnimDatabaseFrameRanges& FrameRanges) const override;
+};
+
+/** Finds all the frames at the end of the given frame ranges */
+UCLASS(BlueprintType, Blueprintable, meta = (DisplayName = "Frame Range Ends"))
+class UAnimDatabaseFramesFunction_FrameRangeEnds : public UAnimDatabaseFramesFunction
+{
+	GENERATED_BODY()
+
+public:
+
+	UPROPERTY(EditAnywhere, Instanced, Category = "Settings")
+	TObjectPtr<UAnimDatabaseFrameRangesFunction> Function;
+
+public:
+
+	UE_API virtual FAnimDatabaseFrames MakeFrames_Implementation(const UAnimDatabase* Database, const FAnimDatabaseFrameRanges& FrameRanges) const override;
+};
+
 /** Entry specifying a frames function */
 USTRUCT(BlueprintType)
 struct FAnimDatabaseFramesEntry
@@ -442,7 +510,7 @@ public:
 	UFUNCTION(BlueprintPure, Category = "AnimDatabase")
 	static UE_API FAnimDatabaseFrameRanges MakeFrameRangesFromSequences(const UAnimDatabase* Database, const TArray<UAnimSequence*>& AnimSequences);
 	static UE_API FAnimDatabaseFrameRanges MakeFrameRangesFromSequencesArrayView(const UAnimDatabase* Database, const TArrayView<UAnimSequence* const> AnimSequences);
-
+	
 	/** Make a set of FrameRanges using the given UAnimDatabaseFrameRangesFunction object */
 	UFUNCTION(BlueprintPure, Category = "AnimDatabase")
 	static UE_API FAnimDatabaseFrameRanges MakeFrameRangesFromFunction(const UAnimDatabase* Database, const FAnimDatabaseFrameRanges& FrameRanges, const UAnimDatabaseFrameRangesFunction* Function);
@@ -572,6 +640,22 @@ public:
 	UFUNCTION(BlueprintPure, Category = "AnimDatabase", meta = (CompactNodeTitle = "NOT"))
 	static UE_API FAnimDatabaseFrameRanges FrameRangesNot(const UAnimDatabase* Database, const FAnimDatabaseFrameRanges& FrameRanges);
 
+	/** Returns the frame ranges that bound the given frame ranges in each sequence */
+	UFUNCTION(BlueprintPure, Category = "AnimDatabase", meta = (CompactNodeTitle = "BOUNDS"))
+	static UE_API FAnimDatabaseFrameRanges FrameRangesBounds(const FAnimDatabaseFrameRanges& FrameRanges);
+
+	/** Returns the frame ranges that bound the given frames in each sequence */
+	UFUNCTION(BlueprintPure, Category = "AnimDatabase", meta = (CompactNodeTitle = "BOUNDS"))
+	static UE_API FAnimDatabaseFrameRanges FrameRangesBoundsFrames(const FAnimDatabaseFrames& Frames);
+
+	/** Returns all the frame ranges from A which contain any part of the frame ranges from B */
+	UFUNCTION(BlueprintPure, Category = "AnimDatabase", meta = (CompactNodeTitle = "INTERSECTS"))
+	static UE_API FAnimDatabaseFrameRanges FrameRangesIntersects(const FAnimDatabaseFrameRanges& A, const FAnimDatabaseFrameRanges& B);
+
+	/** Returns all the frame ranges from A which contain any part of the frames from B */
+	UFUNCTION(BlueprintPure, Category = "AnimDatabase", meta = (CompactNodeTitle = "INTERSECTS"))
+	static UE_API FAnimDatabaseFrameRanges FrameRangesIntersectsFrames(const FAnimDatabaseFrameRanges& A, const FAnimDatabaseFrames& B);
+	
 	/** Return a set of frame ranges with the given amount of frames trimmed from either end. Ranges shorter than the combined TrimStartFrames + TrimEndFrames will be removed. */
 	UFUNCTION(BlueprintPure, Category = "AnimDatabase", meta = (CompactNodeTitle = "TRIM"))
 	static UE_API FAnimDatabaseFrameRanges FrameRangesTrim(const FAnimDatabaseFrameRanges& A, const int32 TrimStartFrames, const int32 TrimEndFrames);
@@ -887,6 +971,22 @@ public:
 	UE_API virtual FAnimDatabaseFrameRanges MakeFrameRanges_Implementation(const UAnimDatabase* Database, const FAnimDatabaseFrameRanges& FrameRanges) const override;
 };
 
+/** Finds all the frame ranges with a given Sync Marker */
+UCLASS(BlueprintType, Blueprintable, meta = (DisplayName = "Sync Marker"))
+class UAnimDatabaseFrameRangesFunction_SyncMarker : public UAnimDatabaseFrameRangesFunction
+{
+	GENERATED_BODY()
+
+public:
+
+	UPROPERTY(EditAnywhere, Category = "Settings")
+	FName SyncMarker = NAME_None;
+
+public:
+
+	UE_API virtual FAnimDatabaseFrameRanges MakeFrameRanges_Implementation(const UAnimDatabase* Database, const FAnimDatabaseFrameRanges& FrameRanges) const override;
+};
+
 /** Finds all the frame ranges with a given Anim Notify State Class */
 UCLASS(BlueprintType, Blueprintable, meta = (DisplayName = "Anim Notify State"))
 class UAnimDatabaseFrameRangesFunction_AnimNotifyState : public UAnimDatabaseFrameRangesFunction
@@ -984,7 +1084,7 @@ public:
 };
 
 /** Finds all the frame ranges with the given sub-string in the package */
-UCLASS(BlueprintType, Blueprintable, meta = (DisplayName = "AssetPackage Contains"))
+UCLASS(BlueprintType, Blueprintable, meta = (DisplayName = "Asset Package Contains"))
 class UAnimDatabaseFrameRangesFunction_AssetPackageContains : public UAnimDatabaseFrameRangesFunction
 {
 	GENERATED_BODY()
@@ -1127,6 +1227,70 @@ public:
 	UE_API virtual FAnimDatabaseFrameRanges MakeFrameRanges_Implementation(const UAnimDatabase* Database, const FAnimDatabaseFrameRanges& FrameRanges) const override;
 };
 
+/** Finds all the frame ranges which bounds another set of ranges */
+UCLASS(BlueprintType, Blueprintable, meta = (DisplayName = "Bounds"))
+class UAnimDatabaseFrameRangesFunction_Bounds : public UAnimDatabaseFrameRangesFunction
+{
+	GENERATED_BODY()
+
+public:
+
+	UPROPERTY(EditAnywhere, Instanced, Category = "Settings")
+	TObjectPtr<UAnimDatabaseFrameRangesFunction> Function;
+
+public:
+
+	UE_API virtual FAnimDatabaseFrameRanges MakeFrameRanges_Implementation(const UAnimDatabase* Database, const FAnimDatabaseFrameRanges& FrameRanges) const override;
+};
+
+/** Finds all the frame ranges which contain any part of the provided frame ranges */
+UCLASS(BlueprintType, Blueprintable, meta = (DisplayName = "Intersects"))
+class UAnimDatabaseFrameRangesFunction_Intersects : public UAnimDatabaseFrameRangesFunction
+{
+	GENERATED_BODY()
+
+public:
+
+	UPROPERTY(EditAnywhere, Instanced, Category = "Settings")
+	TObjectPtr<UAnimDatabaseFrameRangesFunction> Function;
+
+public:
+
+	UE_API virtual FAnimDatabaseFrameRanges MakeFrameRanges_Implementation(const UAnimDatabase* Database, const FAnimDatabaseFrameRanges& FrameRanges) const override;
+};
+
+/** Finds all the frame ranges which contain any part of the provided frames */
+UCLASS(BlueprintType, Blueprintable, meta = (DisplayName = "Intersects Frames"))
+class UAnimDatabaseFrameRangesFunction_IntersectsFrames : public UAnimDatabaseFrameRangesFunction
+{
+	GENERATED_BODY()
+
+public:
+
+	UPROPERTY(EditAnywhere, Instanced, Category = "Settings")
+	TObjectPtr<UAnimDatabaseFramesFunction> Function;
+
+public:
+
+	UE_API virtual FAnimDatabaseFrameRanges MakeFrameRanges_Implementation(const UAnimDatabase* Database, const FAnimDatabaseFrameRanges& FrameRanges) const override;
+};
+
+/** Finds all the frame ranges which bounds another set of frames */
+UCLASS(BlueprintType, Blueprintable, meta = (DisplayName = "Bounds Frames"))
+class UAnimDatabaseFrameRangesFunction_BoundsFrames : public UAnimDatabaseFrameRangesFunction
+{
+	GENERATED_BODY()
+
+public:
+
+	UPROPERTY(EditAnywhere, Instanced, Category = "Settings")
+	TObjectPtr<UAnimDatabaseFramesFunction> Function;
+
+public:
+
+	UE_API virtual FAnimDatabaseFrameRanges MakeFrameRanges_Implementation(const UAnimDatabase* Database, const FAnimDatabaseFrameRanges& FrameRanges) const override;
+};
+
 /** Frame ranges consisting of single frames spaced at a regular interval */
 UCLASS(BlueprintType, Blueprintable, meta = (DisplayName = "Interval"))
 class UAnimDatabaseFrameRangesFunction_Interval : public UAnimDatabaseFrameRangesFunction
diff --git a/Engine/Plugins/Experimental/Animation/AnimDatabase/Source/AnimDatabaseEditor/Private/AnimDatabaseEditorMode.cpp b/Engine/Plugins/Experimental/Animation/AnimDatabase/Source/AnimDatabaseEditor/Private/AnimDatabaseEditorMode.cpp
--- a/Engine/Plugins/Experimental/Animation/AnimDatabase/Source/AnimDatabaseEditor/Private/AnimDatabaseEditorMode.cpp
+++ b/Engine/Plugins/Experimental/Animation/AnimDatabase/Source/AnimDatabaseEditor/Private/AnimDatabaseEditorMode.cpp
@@ -156,7 +156,7 @@ namespace UE::AnimDatabase::Editor
 
 		// Update Query
 
-		Database->Query->RangeIndentifierColorSeed = Database->ViewportSettings->RangeIdentifierColorSeed;
+		Database->Query->RangeIdentifierColorSeed = Database->ViewportSettings->RangeIdentifierColorSeed;
 		const bool bQueryUpdated = Database->Query->Update();
 
 		if (bQueryUpdated)
@@ -230,6 +230,7 @@ namespace UE::AnimDatabase::Editor
 
 			// Adjust the timeline for the newly selected ranges
 
+
 			if (ActiveRanges.Num() == 0)
 			{
 				EditorToolKit->TimelineModel->SetTimelineFrameRate(Database->GetFrameRate());
@@ -251,6 +252,8 @@ namespace UE::AnimDatabase::Editor
 			}
 			else
 			{
+				const FFrameTime CurrentTimelineTime = EditorToolKit->TimelineModel->GetTimelineFrameTime();
+
 				int32 RangeLength = 1;
 				for (const int32 ActiveRangeIdx : ActiveRanges)
 				{
@@ -267,9 +270,7 @@ namespace UE::AnimDatabase::Editor
 
 		// If selection has changed or additional frames/frame ranges has changed update tracks
 
-		const bool bAdditionalTimelineObjectsUpdateRequired = Database->Query->UpdateAdditionalTimelineObjects();
-
-		if (bSelectedRangesUpdated || bAdditionalTimelineObjectsUpdateRequired)
+		if (bSelectedRangesUpdated)
 		{
 			EditorToolKit->TracksModel->ResetTracks();
 
diff --git a/Engine/Plugins/Experimental/Animation/AnimDatabase/Source/AnimDatabaseEditor/Private/AnimDatabaseEditorToolkit.cpp b/Engine/Plugins/Experimental/Animation/AnimDatabase/Source/AnimDatabaseEditor/Private/AnimDatabaseEditorToolkit.cpp
--- a/Engine/Plugins/Experimental/Animation/AnimDatabase/Source/AnimDatabaseEditor/Private/AnimDatabaseEditorToolkit.cpp
+++ b/Engine/Plugins/Experimental/Animation/AnimDatabase/Source/AnimDatabaseEditor/Private/AnimDatabaseEditorToolkit.cpp
@@ -134,6 +134,7 @@ namespace UE::AnimDatabase::Editor
 		// Put the Query category first
 
 		IDetailCategoryBuilder& QueryCategory = DetailBuilder.EditCategory("Query", LOCTEXT("QueryCategory", "Query"), ECategoryPriority::Important);
+		QueryCategory.AddProperty(GET_MEMBER_NAME_CHECKED(UAnimDatabaseQuery, FrameRanges), UAnimDatabaseQuery::StaticClass());
 
 		// Put the frame ranges returned by the query next
 
@@ -743,7 +744,6 @@ namespace UE::AnimDatabase::Editor
 
 		// Reset all the transient properties and update when we are opening a new asset
 		InDatabase->Query->Database = InDatabase;
-		InDatabase->Query->ResetTransient();
 		InDatabase->Query->Update();
 
 		// Create Preview Scene
diff --git a/Engine/Plugins/Experimental/Animation/AnimDatabase/Source/AnimDatabaseEditor/Private/SAnimDatabaseEditorTimeline.cpp b/Engine/Plugins/Experimental/Animation/AnimDatabase/Source/AnimDatabaseEditor/Private/SAnimDatabaseEditorTimeline.cpp
--- a/Engine/Plugins/Experimental/Animation/AnimDatabase/Source/AnimDatabaseEditor/Private/SAnimDatabaseEditorTimeline.cpp
+++ b/Engine/Plugins/Experimental/Animation/AnimDatabase/Source/AnimDatabaseEditor/Private/SAnimDatabaseEditorTimeline.cpp
@@ -2422,8 +2422,8 @@ namespace UE::AnimDatabase::Editor
 		const double OutputViewSize = LocalViewRangeMax - LocalViewRangeMin;
 		const double OutputChange = OutputViewSize * InDelta;
 
-		double NewViewOutputMin = LocalViewRangeMin - (OutputChange * MousePositionFraction);
-		double NewViewOutputMax = LocalViewRangeMax + (OutputChange * (1.0f - MousePositionFraction));
+		double NewViewOutputMin = LocalViewRangeMin - FMath::CeilToFloat(OutputChange * MousePositionFraction);
+		double NewViewOutputMax = LocalViewRangeMax + FMath::CeilToFloat(OutputChange * (1.0f - MousePositionFraction));
 		if (NewViewOutputMin < NewViewOutputMax)
 		{
 			ClampViewRange(NewViewOutputMin, NewViewOutputMax);
@@ -2886,7 +2886,7 @@ namespace UE::AnimDatabase::Editor
 	{
 		// First try built-in spacing.
 		const bool bResult = InFrameRate.ComputeGridSpacing(PixelsPerSecond, OutMajorInterval, OutMinorDivisions, MinTickPx, DesiredMajorTickPx);
-		if (!bResult || OutMajorInterval == 1.0)
+		if (!bResult || OutMajorInterval <= 1.0)
 		{
 			if (PixelsPerSecond <= 0.0f)
 			{
@@ -2943,7 +2943,7 @@ namespace UE::AnimDatabase::Editor
 				const int32 Base = CommonBases[BaseIndex];
 
 				const int32 MajorIntervalFrames = FMath::CeilToInt(Scale / (float)Base) * Base;
-				OutMajorInterval = MajorIntervalFrames * InFrameRate.AsInterval();
+				OutMajorInterval = FMath::Max(MajorIntervalFrames * InFrameRate.AsInterval(), 1.0f);
 
 				// Find the lowest number of divisions we can show that's larger than the minimum tick size.
 				OutMinorDivisions = 0;
diff --git a/Engine/Plugins/Experimental/Animation/AnimGen/Source/AnimGen/Private/AnimGenAutoEncoder.cpp b/Engine/Plugins/Experimental/Animation/AnimGen/Source/AnimGen/Private/AnimGenAutoEncoder.cpp
--- a/Engine/Plugins/Experimental/Animation/AnimGen/Source/AnimGen/Private/AnimGenAutoEncoder.cpp
+++ b/Engine/Plugins/Experimental/Animation/AnimGen/Source/AnimGen/Private/AnimGenAutoEncoder.cpp
@@ -50,42 +50,32 @@ USkeleton* UAnimGenAutoEncoderTrainingSettings::GetSkeleton(bool& bInvalidSkelet
 
 #if WITH_EDITOR
 
-void UAnimGenAutoEncoderTrainingSettings::ResetTransient()
+void UAnimGenAutoEncoderTrainingSettings::ForceRefresh()
 {
 	DatabaseContentHash = 0;
-	QueryRanges = FAnimDatabaseFrameRanges();
-	ActiveFrameRanges = nullptr;
-	QueryEntries.Reset();
-	ActiveRanges.Reset();
-	SelectedRanges.Reset();
-	SelectedFrameRanges = FAnimDatabaseFrameRanges();
-	ActiveFrameAttributes.Reset();
-	FrameAttributeObjects.Reset();
-	FrameAttributeNames.Reset();
-	ActiveDebugDrawer = nullptr;
 }
 
 void UAnimGenAutoEncoderTrainingSettings::LoadDatabaseAsync()
 {
-	Database.LoadAsync(FLoadSoftObjectPathAsyncDelegate::CreateLambda([this](const FSoftObjectPath&, UObject*) { ResetTransient(); }));
+	Database.LoadAsync(FLoadSoftObjectPathAsyncDelegate::CreateLambda([this](const FSoftObjectPath&, UObject*) { ForceRefresh(); }));
 }
 
 void UAnimGenAutoEncoderTrainingSettings::PostEditChangeProperty(struct FPropertyChangedEvent& Event)
 {
 	if (Event.GetMemberPropertyName() == GET_MEMBER_NAME_CHECKED(UAnimGenAutoEncoderTrainingSettings, Database))
 	{
-		ResetTransient();
+		ForceRefresh();
 		LoadDatabaseAsync();
 	}
 }
 
 void UAnimGenAutoEncoderTrainingSettings::UpdateQueryEntries()
 {
-	QueryEntries.Reset();
-
 	if (const UAnimDatabase* DatabaseObject = Database.Get())
 	{
 		check(QueryRanges.IsValid());
+		QueryEntries.SetNum(QueryRanges.FrameRangeSet->GetTotalRangeNum());
+		int32 QueryEntryIdx = 0;
 
 		const int32 EntryNum = QueryRanges.FrameRangeSet->GetEntryNum();
 
@@ -99,12 +89,25 @@ void UAnimGenAutoEncoderTrainingSettings::UpdateQueryEntries()
 
 			for (int32 EntryRangeIdx = 0; EntryRangeIdx < RangeNum; EntryRangeIdx++)
 			{
-				TSharedPtr<UE::AnimDatabase::Editor::FQueryEntry> QueryEntry = MakeShared<UE::AnimDatabase::Editor::FQueryEntry>();
-				QueryEntry->Sequence = AnimSequence;
-				QueryEntry->StartFrame = QueryRanges.FrameRangeSet->GetEntryRangeStart(EntryIdx, EntryRangeIdx);
-				QueryEntry->StopFrame = QueryEntry->StartFrame + QueryRanges.FrameRangeSet->GetEntryRangeLength(EntryIdx, EntryRangeIdx) - 1;
-				QueryEntry->bIsMirrored = bIsMirrored;
-				QueryEntry->bIsSelected = false;
+				TSharedPtr<UE::AnimDatabase::Editor::FQueryEntry>& QueryEntry = QueryEntries[QueryEntryIdx];
+				if (!QueryEntry.IsValid()) { QueryEntry = MakeShared<UE::AnimDatabase::Editor::FQueryEntry>(); }
+
+				const int32 NewStartFrame = QueryRanges.FrameRangeSet->GetEntryRangeStart(EntryIdx, EntryRangeIdx);
+				const int32 NewStopFrame = QueryRanges.FrameRangeSet->GetEntryRangeStop(EntryIdx, EntryRangeIdx) - 1;
+
+				// If the entry has changed then reset selection
+
+				if (QueryEntry->Sequence != AnimSequence ||
+					QueryEntry->StartFrame != NewStartFrame ||
+					QueryEntry->StopFrame != NewStopFrame ||
+					QueryEntry->bIsMirrored != bIsMirrored)
+				{
+					QueryEntry->Sequence = AnimSequence;
+					QueryEntry->StartFrame = NewStartFrame;
+					QueryEntry->StopFrame = NewStopFrame;
+					QueryEntry->bIsMirrored = bIsMirrored;
+					QueryEntry->bIsSelected = false;
+				}
 
 				// Here we use a hash of the sequence, start, and stop frames to produce the color
 
@@ -124,9 +127,15 @@ void UAnimGenAutoEncoderTrainingSettings::UpdateQueryEntries()
 
 				QueryEntry->Color = FLinearColor::MakeFromHSV8(Hue, Saturation, Brightness);
 
-				QueryEntries.Add(QueryEntry);
+				QueryEntryIdx++;
 			}
 		}
+
+		check(QueryEntryIdx == QueryEntries.Num());
+	}
+	else
+	{
+		QueryEntries.Reset();
 	}
 }
 
@@ -169,8 +178,6 @@ bool UAnimGenAutoEncoderTrainingSettings::Update()
 			check(QueryRanges.IsValid());
 
 			UpdateQueryEntries();
-			SelectedRanges.Empty();
-			SelectedFrameRanges = UAnimDatabaseFrameRangesLibrary::MakeEmptyFrameRanges();
 		}
 
 		check(QueryRanges.IsValid());
@@ -191,7 +198,7 @@ bool UAnimGenAutoEncoderTrainingSettings::Update()
 
 		// Update Selection
 
-		ActiveRanges.Reset();
+		SelectedRanges.Reset();
 
 		const int32 QueryRangeNum = QueryEntries.Num();
 
@@ -199,32 +206,18 @@ bool UAnimGenAutoEncoderTrainingSettings::Update()
 		{
 			if (QueryEntries[QueryRangeIdx]->bIsSelected)
 			{
-				ActiveRanges.Add(QueryRangeIdx);
+				SelectedRanges.Add(QueryRangeIdx);
 			}
 		}
 
-		if (bQueryUpdateRequired || SelectedRanges != ActiveRanges)
-		{
-			SelectedRanges = ActiveRanges;
-			SelectedFrameRanges = UAnimDatabaseFrameRangesLibrary::FrameRangesGatherRangesFromIndexSet(QueryRanges, SelectedRanges);
-		}
-
 		// Return if query update required
 
 		return bQueryUpdateRequired;
 	}
-	else
-	{
-		if (DatabaseContentHash != 0)
-		{
-			ResetTransient();
-			return true;
-		}
-		else
-		{
-			return false;
-		}
-	}
+
+	const bool bUpdateRequired = DatabaseContentHash != 0;
+	DatabaseContentHash = 0;
+	return bUpdateRequired;
 }
 
 bool UAnimGenAutoEncoderTrainingSettings::UpdateFrameAttributeObjects()
diff --git a/Engine/Plugins/Experimental/Animation/AnimGen/Source/AnimGen/Private/AnimGenBehavior.cpp b/Engine/Plugins/Experimental/Animation/AnimGen/Source/AnimGen/Private/AnimGenBehavior.cpp
--- a/Engine/Plugins/Experimental/Animation/AnimGen/Source/AnimGen/Private/AnimGenBehavior.cpp
+++ b/Engine/Plugins/Experimental/Animation/AnimGen/Source/AnimGen/Private/AnimGenBehavior.cpp
@@ -670,6 +670,152 @@ void UAnimGenBehavior_TrajectoryInteraction::DrawDebugControl_Implementation(
 	}
 }
 
+FAnimGenControlSchemaElement UAnimGenBehavior_Multi::SpecifyControl_Implementation(FAnimGenControlSchema& InControlSchema) const
+{
+	const int32 TagNum = TagRanges.Num();
+
+	TArray<FName, TInlineAllocator<32>> TagNames;
+	TagNames.Reserve(TagNum);
+
+	for (int32 TagIdx = 0; TagIdx < TagNum; TagIdx++)
+	{
+		TagNames.AddUnique(TagRanges[TagIdx].Name);
+	}
+
+	TArray<FName, TInlineAllocator<32>> SubBehaviorNames;
+	TArray<FAnimGenControlSchemaElement, TInlineAllocator<32>> SubBehaviorElements;
+
+	const int32 BehaviorNum = Behaviors.Num();
+
+	for (int32 BehaviorIdx = 0 ; BehaviorIdx < BehaviorNum; BehaviorIdx++)
+	{
+		if (Behaviors[BehaviorIdx].Behavior)
+		{
+			SubBehaviorNames.Add(Behaviors[BehaviorIdx].Name);
+			SubBehaviorElements.Add(UAnimGenControls::SpecifyEncodingControl(InControlSchema, Behaviors[BehaviorIdx].Behavior->SpecifyControl(InControlSchema)));
+		}
+	}
+
+	return UAnimGenControls::SpecifyStructControlFromArrayViews(InControlSchema,
+		{
+			TEXT("Behavior"),
+			TEXT("Tags"),
+		},
+		{
+			UAnimGenControls::SpecifyExclusiveUnionControlFromArrayViews(InControlSchema, SubBehaviorNames, SubBehaviorElements),
+			UAnimGenControls::SpecifyNamedInclusiveDiscreteControlFromArrayView(InControlSchema, TagNames),
+		});
+}
+
+void UAnimGenBehavior_Multi::MakeControlSets_Implementation(TArray<FAnimGenControlSet>& OutControlSets, FAnimGenControlObject& InControlObject, const UAnimDatabase* InDatabase, const FAnimDatabaseFrameRanges& InFrameRanges) const
+{
+	const int32 TagNum = TagRanges.Num();
+
+	TArray<FName, TInlineAllocator<32>> TagNames;
+	TArray<FAnimDatabaseFrameRanges, TInlineAllocator<32>> TagFrameRanges;
+	TagNames.Reserve(TagNum);
+	TagFrameRanges.Reserve(TagNum);
+
+	for (int32 TagIdx = 0; TagIdx < TagNum; TagIdx++)
+	{
+		if (TagNames.Contains(TagRanges[TagIdx].Name)) { continue; }
+
+		const FAnimDatabaseFrameRanges Ranges = UAnimDatabaseFrameRangesLibrary::MakeFrameRangesFromFunction(InDatabase, InFrameRanges, TagRanges[TagIdx].FrameRanges);
+
+		if (Ranges.IsValid() && !Ranges.FrameRangeSet->IsEmpty())
+		{
+			TagNames.Add(TagRanges[TagIdx].Name);
+			TagFrameRanges.Add(Ranges);
+		}
+	}
+
+	const int32 BehaviorNum = Behaviors.Num();
+
+	for (int32 BehaviorIdx = 0; BehaviorIdx < BehaviorNum; BehaviorIdx++)
+	{
+		if (Behaviors[BehaviorIdx].Behavior)
+		{
+			const FAnimDatabaseFrameRanges BehaviorFramesRanges = Behaviors[BehaviorIdx].FrameRanges ? 
+				UAnimDatabaseFrameRangesLibrary::MakeFrameRangesFromFunction(InDatabase, InFrameRanges, Behaviors[BehaviorIdx].FrameRanges) : InFrameRanges;
+
+			TArray<FAnimGenControlSet> SubBehaviorControlSets;
+			Behaviors[BehaviorIdx].Behavior->MakeControlSets(SubBehaviorControlSets, InControlObject, InDatabase, BehaviorFramesRanges);
+
+			for (const FAnimGenControlSet& ControlSet : SubBehaviorControlSets)
+			{
+				OutControlSets.Add(UAnimGenControls::MakeStructControlSetFromArrayViews(InControlObject,
+					{
+						TEXT("Behavior"),
+						TEXT("Tags")
+					},
+					{
+						UAnimGenControls::MakeExclusiveUnionControlSet(InControlObject, Behaviors[BehaviorIdx].Name, UAnimGenControls::MakeEncodingControlSet(InControlObject, ControlSet)),
+						UAnimGenControls::MakeNamedInclusiveDiscreteControlSetFromFrameRangesArrayViews(InControlObject, UAnimGenControls::ControlSetFrameRanges(ControlSet), TagNames, TagFrameRanges),
+					}));
+			}
+		}
+	}
+}
+
+void UAnimGenBehavior_Multi::DrawDebugControl_Implementation(
+	const FDebugDrawer& Drawer,
+	const FAnimGenControlObject& InControlObject,
+	const FAnimGenControlObjectElement InControlObjectElement,
+	const FAnimDatabasePoseState& InPoseState,
+	const UAnimDatabase* InDatabase,
+	const UAnimGenAutoEncoder* InAutoEncoder,
+	const int32 SequenceIdx,
+	const float SequenceTime,
+	const int32 RangeStart,
+	const int32 RangeLength,
+	const FLinearColor& IdentifierColor) const
+{
+	FAnimGenControlObjectElement BehaviorControl, TagsControl;
+	UAnimGenControls::GetStructControlElement(BehaviorControl, InControlObject, InControlObjectElement, TEXT("Behavior"));
+	UAnimGenControls::GetStructControlElement(TagsControl, InControlObject, InControlObjectElement, TEXT("Tags"));
+
+	const FTransform RootTransform = UAnimDatabasePoseStateLibrary::PoseStateRootTransform(InPoseState);
+
+	UAnimGenControls::DrawDebugNamedInclusiveDiscreteControl(Drawer, InControlObject, TagsControl, TEXT(""), RootTransform.GetLocation() + FVector(0.0, 0.0, 260.0f), FRotator(0.0f, -90.0f, 0.0f), 10.0f, IdentifierColor);
+
+	FName BehaviorName = NAME_None;
+	FAnimGenControlObjectElement EncodingElement, BehaviorElement;
+	UAnimGenControls::GetExclusiveUnionControl(BehaviorName, EncodingElement, InControlObject, BehaviorControl);
+	UAnimGenControls::GetEncodingControl(BehaviorElement, InControlObject, EncodingElement);
+
+	FDrawDebugStringSettings Settings;
+	Settings.Height = 10.0f;
+	Settings.bCenterHorizontally = true;
+	Settings.bCenterVertically = true;
+
+	FDrawDebugLineStyle LineStyle;
+	LineStyle.Color = IdentifierColor;
+
+	UDrawDebugLibrary::DrawDebugName(Drawer, BehaviorName, RootTransform.GetLocation() + FVector(0.0, 0.0, 240.0f), FRotator(0.0f, -90.0f, 0.0f), LineStyle, true, Settings);
+
+	const int32 BehaviorNum = Behaviors.Num();
+
+	for (int32 BehaviorIdx = 0; BehaviorIdx < BehaviorNum; BehaviorIdx++)
+	{
+		if (Behaviors[BehaviorIdx].Behavior && Behaviors[BehaviorIdx].Name == BehaviorName)
+		{
+			Behaviors[BehaviorIdx].Behavior->DrawDebugControl(
+				Drawer, 
+				InControlObject, 
+				BehaviorElement,
+				InPoseState, 
+				InDatabase, 
+				InAutoEncoder, 
+				SequenceIdx, 
+				SequenceTime, 
+				RangeStart, 
+				RangeLength, 
+				IdentifierColor);
+			break;
+		}
+	}
+}
+
 #endif
 
 FAnimGenControlObjectElement UAnimGenBehaviorLibrary::MakeUncontrolledBehaviorControl(FAnimGenControlObject& InControlObject)
@@ -803,6 +949,32 @@ FAnimGenControlObjectElement UAnimGenBehaviorLibrary::MakeTrajectoryInteractionB
 		});
 }
 
+FAnimGenControlObjectElement UAnimGenBehaviorLibrary::MakeMultiBehaviorControl(
+	FAnimGenControlObject& InControlObject,
+	const FName Behavior,
+	const FAnimGenControlObjectElement& BehaviorControlElement,
+	const TArray<FName>& DesiredTags)
+{
+	return MakeMultiBehaviorControlFromArrayView(InControlObject, Behavior, BehaviorControlElement, DesiredTags);
+}
+
+FAnimGenControlObjectElement UAnimGenBehaviorLibrary::MakeMultiBehaviorControlFromArrayView(
+	FAnimGenControlObject& InControlObject,
+	const FName Behavior,
+	const FAnimGenControlObjectElement& BehaviorControlElement,
+	const TArrayView<const FName> DesiredTags)
+{
+	return UAnimGenControls::MakeStructControlFromArrayViews(InControlObject,
+		{
+			TEXT("Behavior"),
+			TEXT("Tags"),
+		},
+		{
+			UAnimGenControls::MakeExclusiveUnionControl(InControlObject, Behavior, UAnimGenControls::MakeEncodingControl(InControlObject, BehaviorControlElement)),
+			UAnimGenControls::MakeNamedInclusiveDiscreteControlFromArrayView(InControlObject, DesiredTags),
+		});
+}
+
 
 bool UAnimGenBehaviorLibrary::MakeRandomBool(int32& State)
 {
diff --git a/Engine/Plugins/Experimental/Animation/AnimGen/Source/AnimGen/Private/AnimGenController.cpp b/Engine/Plugins/Experimental/Animation/AnimGen/Source/AnimGen/Private/AnimGenController.cpp
--- a/Engine/Plugins/Experimental/Animation/AnimGen/Source/AnimGen/Private/AnimGenController.cpp
+++ b/Engine/Plugins/Experimental/Animation/AnimGen/Source/AnimGen/Private/AnimGenController.cpp
@@ -22,43 +22,51 @@
 
 #if WITH_EDITOR
 
-void UAnimGenControllerTrainingSettings::ResetTransient()
+void UAnimGenControllerTrainingSettings::ForceRefresh()
 {
 	DatabaseContentHash = 0;
-	ActiveFrameRanges = nullptr;
-	ActiveBehavior = nullptr;
-	ControlSets.Empty();
-	QueryEntries.Empty();
-	SelectedRanges.Empty();
 }
 
 void UAnimGenControllerTrainingSettings::LoadDatabaseAsync()
 {
-	Database.LoadAsync(FLoadSoftObjectPathAsyncDelegate::CreateLambda([this](const FSoftObjectPath&, UObject*) { ResetTransient(); }));
+	Database.LoadAsync(FLoadSoftObjectPathAsyncDelegate::CreateLambda([this](const FSoftObjectPath&, UObject*) { ForceRefresh(); }));
+}
+
+void UAnimGenControllerTrainingSettings::LoadAutoEncoderAsync()
+{
+	AutoEncoder.LoadAsync(FLoadSoftObjectPathAsyncDelegate::CreateLambda([this](const FSoftObjectPath&, UObject*) { ForceRefresh(); }));
 }
 
 void UAnimGenControllerTrainingSettings::PostEditChangeProperty(struct FPropertyChangedEvent& Event)
 {
 	if (Event.GetMemberPropertyName() == GET_MEMBER_NAME_CHECKED(UAnimGenControllerTrainingSettings, Database))
 	{
-		ResetTransient();
+		ForceRefresh();
 		LoadDatabaseAsync();
 	}
 	else if (Event.GetMemberPropertyName() == GET_MEMBER_NAME_CHECKED(UAnimGenControllerTrainingSettings, AutoEncoder))
 	{
-		ResetTransient();
-		LoadDatabaseAsync();
+		ForceRefresh();
+		LoadAutoEncoderAsync();
 	}
 }
 
 void UAnimGenControllerTrainingSettings::UpdateQueryEntries()
 {
-	QueryEntries.Reset();
-
 	if (Database && AutoEncoder)
 	{
 		const int32 ControlSetNum = ControlSets.Num();
 
+		int32 TotalRangeNum = 0;
+		for (int32 ControlSetIdx = 0; ControlSetIdx < ControlSetNum; ControlSetIdx++)
+		{
+			check(ControlSets[ControlSetIdx].IsValid());
+			TotalRangeNum += ControlSets[ControlSetIdx].Data->FrameRangeSet.GetTotalRangeNum();
+		}
+
+		QueryEntries.SetNum(TotalRangeNum);
+		int32 QueryEntryIdx = 0;
+
 		for (int32 ControlSetIdx = 0; ControlSetIdx < ControlSetNum; ControlSetIdx++)
 		{
 			const UE::Learning::FFrameRangeSet& FrameRangeSet = ControlSets[ControlSetIdx].Data->FrameRangeSet;
@@ -75,12 +83,25 @@ void UAnimGenControllerTrainingSettings::UpdateQueryEntries()
 
 				for (int32 EntryRangeIdx = 0; EntryRangeIdx < RangeNum; EntryRangeIdx++)
 				{
-					TSharedPtr<UE::AnimDatabase::Editor::FQueryEntry> QueryEntry = MakeShared<UE::AnimDatabase::Editor::FQueryEntry>();
-					QueryEntry->Sequence = AnimSequence;
-					QueryEntry->StartFrame = FrameRangeSet.GetEntryRangeStart(EntryIdx, EntryRangeIdx);
-					QueryEntry->StopFrame = QueryEntry->StartFrame + FrameRangeSet.GetEntryRangeLength(EntryIdx, EntryRangeIdx) - 1;
-					QueryEntry->bIsMirrored = bIsMirrored;
-					QueryEntry->bIsSelected = false;
+					TSharedPtr<UE::AnimDatabase::Editor::FQueryEntry>& QueryEntry = QueryEntries[QueryEntryIdx];
+					if (!QueryEntry.IsValid()) { QueryEntry = MakeShared<UE::AnimDatabase::Editor::FQueryEntry>(); }
+
+					const int32 NewStartFrame = FrameRangeSet.GetEntryRangeStart(EntryIdx, EntryRangeIdx);
+					const int32 NewStopFrame = FrameRangeSet.GetEntryRangeStop(EntryIdx, EntryRangeIdx) - 1;
+
+					// If the entry has changed then reset selection
+
+					if (QueryEntry->Sequence != AnimSequence ||
+						QueryEntry->StartFrame != NewStartFrame ||
+						QueryEntry->StopFrame != NewStopFrame ||
+						QueryEntry->bIsMirrored != bIsMirrored)
+					{
+						QueryEntry->Sequence = AnimSequence;
+						QueryEntry->StartFrame = NewStartFrame;
+						QueryEntry->StopFrame = NewStopFrame;
+						QueryEntry->bIsMirrored = bIsMirrored;
+						QueryEntry->bIsSelected = false;
+					}
 
 					// Here we use a hash of the sequence, start, and stop frames to produce the color
 
@@ -101,10 +122,16 @@ void UAnimGenControllerTrainingSettings::UpdateQueryEntries()
 
 					QueryEntry->Color = FLinearColor::MakeFromHSV8(Hue, Saturation, Brightness);
 
-					QueryEntries.Add(QueryEntry);
+					QueryEntryIdx++;
 				}
 			}
 		}
+
+		check(QueryEntryIdx == QueryEntries.Num());
+	}
+	else
+	{
+		QueryEntries.Reset();
 	}
 }
 
@@ -182,7 +209,6 @@ bool UAnimGenControllerTrainingSettings::Update()
 			}
 
 			UpdateQueryEntries();
-			SelectedRanges.Empty();
 		}
 
 		// Update Selection
@@ -201,18 +227,10 @@ bool UAnimGenControllerTrainingSettings::Update()
 
 		return bControlSetsUpdateRequired;
 	}
-	else
-	{
-		if (DatabaseContentHash != 0)
-		{
-			ResetTransient();
-			return true;
-		}
-		else
-		{
-			return false;
-		}
-	}
+
+	const bool bUpdateRequired = DatabaseContentHash != 0;
+	DatabaseContentHash = 0;
+	return bUpdateRequired;
 }
 
 void UAnimGenControllerTestingSettings::LoadDatabaseAsync()
diff --git a/Engine/Plugins/Experimental/Animation/AnimGen/Source/AnimGen/Public/AnimGenAutoEncoder.h b/Engine/Plugins/Experimental/Animation/AnimGen/Source/AnimGen/Public/AnimGenAutoEncoder.h
--- a/Engine/Plugins/Experimental/Animation/AnimGen/Source/AnimGen/Public/AnimGenAutoEncoder.h
+++ b/Engine/Plugins/Experimental/Animation/AnimGen/Source/AnimGen/Public/AnimGenAutoEncoder.h
@@ -300,9 +300,9 @@ public:
 
 #if WITH_EDITOR
 
-	/** Resets all of the transient properties/members. */
+	/** Forces the Update function to refresh the internal state */
 	UFUNCTION(CallInEditor, Category = "Data", DisplayName = "Refresh")
-	UE_API void ResetTransient();
+	UE_API void ForceRefresh();
 
 	/** Load the current database. */
 	UE_API void LoadDatabaseAsync();
@@ -335,15 +335,9 @@ public:
 	/** List of UI entries for the current query */
 	TArray<TSharedPtr<UE::AnimDatabase::Editor::FQueryEntry>> QueryEntries;
 
-	/** Temporary array of active ranges used to compare selected ranges */
-	TArray<int32> ActiveRanges;
-
 	/** Indices of all the currently selected ranges */
 	TArray<int32> SelectedRanges;
 
-	/** The frame ranges of the selected ranges */
-	FAnimDatabaseFrameRanges SelectedFrameRanges;
-
 	/** Seed used for the colors of the range identifiers */
 	int32 RangeIndentifierColorSeed = 1234;
 
diff --git a/Engine/Plugins/Experimental/Animation/AnimGen/Source/AnimGen/Public/AnimGenBehavior.h b/Engine/Plugins/Experimental/Animation/AnimGen/Source/AnimGen/Public/AnimGenBehavior.h
--- a/Engine/Plugins/Experimental/Animation/AnimGen/Source/AnimGen/Public/AnimGenBehavior.h
+++ b/Engine/Plugins/Experimental/Animation/AnimGen/Source/AnimGen/Public/AnimGenBehavior.h
@@ -526,6 +526,72 @@ public:
 #endif
 };
 
+/** Entry specifying a behavior */
+USTRUCT(BlueprintType)
+struct FAnimGenBehaviorEntry
+{
+	GENERATED_BODY()
+
+public:
+
+	/** Behavior name */
+	UPROPERTY(EditAnywhere, Category = "Entry")
+	FName Name = NAME_None;
+
+	/** Behavior object */
+	UPROPERTY(EditAnywhere, Instanced, Category = "Entry")
+	TObjectPtr<UAnimGenBehavior> Behavior;
+
+	/** Frame ranges to use for behavior */
+	UPROPERTY(EditAnywhere, Instanced, Category = "Entry")
+	TObjectPtr<UAnimDatabaseFrameRangesFunction> FrameRanges;
+};
+
+/** Behavior made up of multiple other behaviors - user can also provide some tags to control the style. */
+UCLASS(BlueprintType, Blueprintable, meta = (DisplayName = "Multi"))
+class UAnimGenBehavior_Multi : public UAnimGenBehavior
+{
+	GENERATED_BODY()
+
+#if WITH_EDITORONLY_DATA
+
+public:
+
+	/** Different potential sub-behaviors */
+	UPROPERTY(EditAnywhere, Category = "Behaviors")
+	TArray<FAnimGenBehaviorEntry> Behaviors;
+
+	/** List of tags with associated frame ranges */
+	UPROPERTY(EditAnywhere, Category = "Tags")
+	TArray<FAnimDatabaseFrameRangesEntry> TagRanges;
+
+#endif
+
+#if WITH_EDITOR
+
+public:
+
+	UE_API virtual FAnimGenControlSchemaElement SpecifyControl_Implementation(FAnimGenControlSchema& InControlSchema) const override;
+
+	UE_API virtual void MakeControlSets_Implementation(TArray<FAnimGenControlSet>& OutControlSets, FAnimGenControlObject& InControlObject, const UAnimDatabase* InDatabase, const FAnimDatabaseFrameRanges& InFrameRanges) const override;
+
+	UE_API virtual void DrawDebugControl_Implementation(
+		const FDebugDrawer& Drawer,
+		const FAnimGenControlObject& InControlObject,
+		const FAnimGenControlObjectElement InControlObjectElement,
+		const FAnimDatabasePoseState& InPoseState,
+		const UAnimDatabase* InDatabase,
+		const UAnimGenAutoEncoder* InAutoEncoder,
+		const int32 SequenceIdx,
+		const float SequenceTime,
+		const int32 RangeStart,
+		const int32 RangeLength,
+		const FLinearColor& IdentifierColor) const override;
+
+#endif
+};
+
+
 /** Simple Trajectory type that stores all variables required for simple spring-based prediction. */
 USTRUCT(BlueprintType)
 struct FAnimGenTrajectory
@@ -634,6 +700,20 @@ public:
 		const float InTrajectoryFutureTime = 1.0f,
 		const FVector InTrajectoryForwardVector = FVector::RightVector);
 
+	/** Make a control object element for the Multi Behavior */
+	UFUNCTION(BlueprintPure, Category = "AnimGen", meta = (AutoCreateRefTerm = "DesiredTags"))
+	static UE_API FAnimGenControlObjectElement MakeMultiBehaviorControl(
+		UPARAM(ref) FAnimGenControlObject& InControlObject,
+		const FName Behavior,
+		const FAnimGenControlObjectElement& BehaviorControlElement,
+		const TArray<FName>& DesiredTags);
+
+	static UE_API FAnimGenControlObjectElement MakeMultiBehaviorControlFromArrayView(
+		FAnimGenControlObject& InControlObject,
+		const FName Behavior,
+		const FAnimGenControlObjectElement& BehaviorControlElement,
+		const TArrayView<const FName> DesiredTags);
+
 public:
 
 	/** Make a random bool value */
diff --git a/Engine/Plugins/Experimental/Animation/AnimGen/Source/AnimGen/Public/AnimGenController.h b/Engine/Plugins/Experimental/Animation/AnimGen/Source/AnimGen/Public/AnimGenController.h
--- a/Engine/Plugins/Experimental/Animation/AnimGen/Source/AnimGen/Public/AnimGenController.h
+++ b/Engine/Plugins/Experimental/Animation/AnimGen/Source/AnimGen/Public/AnimGenController.h
@@ -223,13 +223,16 @@ public:
 
 #if WITH_EDITOR
 
-	/** Resets all of the transient properties/members. */
+	/** Forces the Update function to refresh the internal state */
 	UFUNCTION(CallInEditor, Category = "Data", DisplayName = "Refresh")
-	UE_API void ResetTransient();
+	UE_API void ForceRefresh();
 
 	/** Load the current database. */
 	UE_API void LoadDatabaseAsync();
 
+	/** Load the current auto-encoder. */
+	UE_API void LoadAutoEncoderAsync();
+
 	/** We use this async load changes to Database/AutoEncoder */
 	UE_API virtual void PostEditChangeProperty(struct FPropertyChangedEvent& Event) override;
 
@@ -320,27 +323,27 @@ public:
 	int32 DenoiserLayerNum = 10;
 
 	/** The number of hidden units used by the LOD0 network on each layer */
-	UPROPERTY(EditAnywhere, Category = "Network", meta = (ClampMin = "1", UIMin = "1"))
+	UPROPERTY(EditAnywhere, Category = "Network", DisplayName = "LOD0 Hidden Unit Num", meta = (ClampMin = "1", UIMin = "1"))
 	int32 LOD0HiddenUnitNum = 1024;
 
 	/** The number of layers used by the LOD0 network */
-	UPROPERTY(EditAnywhere, Category = "Network", meta = (ClampMin = "2", UIMin = "2"))
+	UPROPERTY(EditAnywhere, Category = "Network", DisplayName = "LOD0 Layer Num", meta = (ClampMin = "2", UIMin = "2"))
 	int32 LOD0LayerNum = 8;
 
 	/** The number of hidden units used by the LOD1 network on each layer */
-	UPROPERTY(EditAnywhere, Category = "Network", meta = (ClampMin = "1", UIMin = "1"))
+	UPROPERTY(EditAnywhere, Category = "Network", DisplayName = "LOD1 Hidden Unit Num", meta = (ClampMin = "1", UIMin = "1"))
 	int32 LOD1HiddenUnitNum = 512;
 
 	/** The number of layers used by the LOD1 network */
-	UPROPERTY(EditAnywhere, Category = "Network", meta = (ClampMin = "2", UIMin = "2"))
+	UPROPERTY(EditAnywhere, Category = "Network", DisplayName = "LOD1 Layer Num", meta = (ClampMin = "2", UIMin = "2"))
 	int32 LOD1LayerNum = 6;
 
 	/** The number of hidden units used by the LOD2 network on each layer */
-	UPROPERTY(EditAnywhere, Category = "Network", meta = (ClampMin = "1", UIMin = "1"))
+	UPROPERTY(EditAnywhere, Category = "Network", DisplayName = "LOD2 Hidden Unit Num", meta = (ClampMin = "1", UIMin = "1"))
 	int32 LOD2HiddenUnitNum = 256;
 
 	/** The number of layers used by the LOD2 network */
-	UPROPERTY(EditAnywhere, Category = "Network", meta = (ClampMin = "2", UIMin = "2"))
+	UPROPERTY(EditAnywhere, Category = "Network", DisplayName = "LOD2 Layer Num", meta = (ClampMin = "2", UIMin = "2"))
 	int32 LOD2LayerNum = 4;
 
 	/**
@@ -610,47 +613,47 @@ public:
 	FFrameRate FrameRate = FFrameRate(60, 1);
 
 	/** Size of the control vector used by this controller */
-	UPROPERTY(VisibleAnywhere, Category = "Controller")
+	UPROPERTY()
 	int32 ControlVectorSize = 0;
 
 	/** Size of the encoded control vector used by this controller */
-	UPROPERTY(VisibleAnywhere, Category = "Controller")
+	UPROPERTY()
 	int32 EncodedControlVectorSize = 0;
 
 	/** Size of the control distribution vector used to clamp controls for this controller */
-	UPROPERTY(VisibleAnywhere, Category = "Controller")
+	UPROPERTY()
 	int32 ControlDistributionVectorSize = 0;
 
 	/** Mean of the encoded pose vectors used for training */
-	UPROPERTY(VisibleAnywhere, Category = "Controller")
+	UPROPERTY()
 	TArray<float> EncodedPoseMeans;
 
 	/** Std of the encoded pose vectors used for training */
-	UPROPERTY(VisibleAnywhere, Category = "Controller")
+	UPROPERTY()
 	TArray<float> EncodedPoseStds;
 
 	/** Max of the encoded pose vectors used for training */
-	UPROPERTY(VisibleAnywhere, Category = "Controller")
+	UPROPERTY()
 	TArray<float> EncodedPoseMaxs;
 
 	/** Min of the encoded pose vectors used for training */
-	UPROPERTY(VisibleAnywhere, Category = "Controller")
+	UPROPERTY()
 	TArray<float> EncodedPoseMins;
 
 	/** Std used for normalizing the encoded pose vectors */
-	UPROPERTY(VisibleAnywhere, Category = "Controller")
+	UPROPERTY()
 	TArray<float> NormalizedPoseStds;
 
 	/** Scale used for normalizing encoded pose vectors */
-	UPROPERTY(VisibleAnywhere, Category = "Controller")
+	UPROPERTY()
 	float EncodedPoseNormalizationScale = 1.0f;
 
 	/** Min of the control vectors used by this controller */
-	UPROPERTY(VisibleAnywhere, Category = "Controller")
+	UPROPERTY()
 	TArray<float> ControlVectorDistributionMins;
 
 	/** Max of the control vectors used by this controller */
-	UPROPERTY(VisibleAnywhere, Category = "Controller")
+	UPROPERTY()
 	TArray<float> ControlVectorDistributionMaxs;
 
 public:
diff --git a/Engine/Plugins/Experimental/Animation/AnimGen/Source/AnimGenEditor/Private/AnimGenEditor.cpp b/Engine/Plugins/Experimental/Animation/AnimGen/Source/AnimGenEditor/Private/AnimGenEditor.cpp
--- a/Engine/Plugins/Experimental/Animation/AnimGen/Source/AnimGenEditor/Private/AnimGenEditor.cpp
+++ b/Engine/Plugins/Experimental/Animation/AnimGen/Source/AnimGenEditor/Private/AnimGenEditor.cpp
@@ -20,6 +20,7 @@ void FAnimGenEditorModule::StartupModule()
 
 	FPropertyEditorModule& PropertyModule = FModuleManager::LoadModuleChecked<FPropertyEditorModule>("PropertyEditor");
 	PropertyModule.RegisterCustomClassLayout("AnimGenAutoEncoder", FOnGetDetailCustomizationInstance::CreateStatic(&UE::AnimGen::Editor::FAutoEncoderDetails::MakeInstance));
+	PropertyModule.RegisterCustomClassLayout("AnimGenController", FOnGetDetailCustomizationInstance::CreateStatic(&UE::AnimGen::Editor::FControllerDetails::MakeInstance));
 	PropertyModule.RegisterCustomClassLayout("AnimGenAutoEncoderTrainingSettings", FOnGetDetailCustomizationInstance::CreateStatic(&UE::AnimGen::Editor::FAutoEncoderTrainingSettingsDetails::MakeInstance));
 	PropertyModule.RegisterCustomClassLayout("AnimGenControllerTrainingSettings", FOnGetDetailCustomizationInstance::CreateStatic(&UE::AnimGen::Editor::FControllerTrainingSettingsDetails::MakeInstance));
 }
@@ -34,6 +35,7 @@ void FAnimGenEditorModule::ShutdownModule()
 	{
 		FPropertyEditorModule& PropertyModule = FModuleManager::GetModuleChecked<FPropertyEditorModule>("PropertyEditor");
 		PropertyModule.UnregisterCustomClassLayout("AnimGenAutoEncoder");
+		PropertyModule.UnregisterCustomClassLayout("AnimGenController");
 		PropertyModule.UnregisterCustomClassLayout("AnimGenAutoEncoderTrainingSettings");
 		PropertyModule.UnregisterCustomClassLayout("AnimGenControllerTrainingSettings");
 	}
diff --git a/Engine/Plugins/Experimental/Animation/AnimGen/Source/AnimGenEditor/Private/AnimGenEditorAutoEncoderMode.cpp b/Engine/Plugins/Experimental/Animation/AnimGen/Source/AnimGenEditor/Private/AnimGenEditorAutoEncoderMode.cpp
--- a/Engine/Plugins/Experimental/Animation/AnimGen/Source/AnimGenEditor/Private/AnimGenEditorAutoEncoderMode.cpp
+++ b/Engine/Plugins/Experimental/Animation/AnimGen/Source/AnimGenEditor/Private/AnimGenEditorAutoEncoderMode.cpp
@@ -245,10 +245,6 @@ namespace UE::AnimGen::Editor
 					FText::AsNumber(MaxCharacterNum),
 					FText::AsNumber(ActiveRanges.Num())));
 			}
-			else
-			{
-				ViewportClient->SetWarningMessage(FText());
-			}
 
 			// Adjust the timeline for the newly selected ranges
 
diff --git a/Engine/Plugins/Experimental/Animation/AnimGen/Source/AnimGenEditor/Private/AnimGenEditorAutoEncoderToolkit.cpp b/Engine/Plugins/Experimental/Animation/AnimGen/Source/AnimGenEditor/Private/AnimGenEditorAutoEncoderToolkit.cpp
--- a/Engine/Plugins/Experimental/Animation/AnimGen/Source/AnimGenEditor/Private/AnimGenEditorAutoEncoderToolkit.cpp
+++ b/Engine/Plugins/Experimental/Animation/AnimGen/Source/AnimGenEditor/Private/AnimGenEditorAutoEncoderToolkit.cpp
@@ -187,6 +187,11 @@ namespace UE::AnimGen::Editor
 		// Put the Data category first
 
 		IDetailCategoryBuilder& DataCategory = DetailBuilder.EditCategory("Data", LOCTEXT("DataCategory", "Data"), ECategoryPriority::Important);
+		DataCategory.AddProperty(GET_MEMBER_NAME_CHECKED(UAnimGenAutoEncoderTrainingSettings, Database), UAnimGenAutoEncoderTrainingSettings::StaticClass());
+		DataCategory.AddProperty(GET_MEMBER_NAME_CHECKED(UAnimGenAutoEncoderTrainingSettings, FrameRanges), UAnimGenAutoEncoderTrainingSettings::StaticClass());
+		DataCategory.AddProperty(GET_MEMBER_NAME_CHECKED(UAnimGenAutoEncoderTrainingSettings, ExcludedBones), UAnimGenAutoEncoderTrainingSettings::StaticClass());
+		DataCategory.AddProperty(GET_MEMBER_NAME_CHECKED(UAnimGenAutoEncoderTrainingSettings, bExcludeVirtualBones), UAnimGenAutoEncoderTrainingSettings::StaticClass());
+		DataCategory.AddProperty(GET_MEMBER_NAME_CHECKED(UAnimGenAutoEncoderTrainingSettings, FrameAttributes), UAnimGenAutoEncoderTrainingSettings::StaticClass());
 
 		// Put the frame ranges returned by the query next
 
@@ -401,7 +406,6 @@ namespace UE::AnimGen::Editor
 
 		// Add Categories
 
-
 		IDetailCategoryBuilder& AutoEncoderCategory = DetailBuilder.EditCategory("AutoEncoder", LOCTEXT("AutoEncoder", "Auto Encoder"), ECategoryPriority::Important);
 		AutoEncoderCategory.AddProperty(GET_MEMBER_NAME_CHECKED(UAnimGenAutoEncoder, TrainedBoneLocations), UAnimGenAutoEncoder::StaticClass());
 		AutoEncoderCategory.AddProperty(GET_MEMBER_NAME_CHECKED(UAnimGenAutoEncoder, TrainedBoneRotations), UAnimGenAutoEncoder::StaticClass());
@@ -410,9 +414,10 @@ namespace UE::AnimGen::Editor
 
 		// Check if Trained Database has been modified
 
-		if (AutoEncoder->IsValid() && UAnimDatabaseFrameRangesLibrary::FrameRangesContentHash(
-			AutoEncoder->TrainingSettings->Database.LoadSynchronous(),
-			AutoEncoder->TrainingSettings->QueryRanges) != AutoEncoder->TrainedContentHash)
+		if (AutoEncoder->IsValid() && AutoEncoder->TrainingSettings->Database.Get() && 
+			UAnimDatabaseFrameRangesLibrary::FrameRangesContentHash(
+				AutoEncoder->TrainingSettings->Database.Get(),
+				AutoEncoder->TrainingSettings->QueryRanges) != AutoEncoder->TrainedContentHash)
 		{
 			FDetailWidgetRow& NeedsReimportErrorRow = AutoEncoderCategory.AddCustomRow(LOCTEXT("ModifiedWarning", "Warning"))
 				.WholeRowContent()
diff --git a/Engine/Plugins/Experimental/Animation/AnimGen/Source/AnimGenEditor/Private/AnimGenEditorControllerMode.cpp b/Engine/Plugins/Experimental/Animation/AnimGen/Source/AnimGenEditor/Private/AnimGenEditorControllerMode.cpp
--- a/Engine/Plugins/Experimental/Animation/AnimGen/Source/AnimGenEditor/Private/AnimGenEditorControllerMode.cpp
+++ b/Engine/Plugins/Experimental/Animation/AnimGen/Source/AnimGenEditor/Private/AnimGenEditorControllerMode.cpp
@@ -494,6 +494,7 @@ namespace UE::AnimGen::Editor
 
 		AnimDatabase::Editor::FViewportClient* ViewportClient = static_cast<AnimDatabase::Editor::FViewportClient*>(EditorViewportClient);
 		check(ViewportClient);
+		ViewportClient->ClearWarningMessage();
 
 		TSharedPtr<FControllerToolkit> EditorToolKit = StaticCastSharedPtr<FControllerToolkit>(ViewportClient->GetAssetEditorToolkit().Pin());
 		if (!EditorToolKit.IsValid()) { return; }
@@ -535,7 +536,21 @@ namespace UE::AnimGen::Editor
 		// Load AutoEncoder
 
 		UAnimGenAutoEncoder* TrainingAutoEncoder = Controller->TrainingSettings->AutoEncoder.Get();
-		if (!TrainingAutoEncoder || !TrainingAutoEncoder->IsValid()) { return; }
+		if (!TrainingAutoEncoder)
+		{
+			ViewportClient->SetWarningMessage(
+				LOCTEXT("LoadingAutoEncoder", "Loading Training AutoEncoder..."));
+			ActiveRanges.Reset();
+			return;
+		}
+
+		if (!TrainingAutoEncoder->IsValid())
+		{
+			ViewportClient->SetWarningMessage(
+				LOCTEXT("AutoEncoderUntrained", "AutoEncoder not Trained..."));
+			ActiveRanges.Reset();
+			return;
+		}
 
 		// Load Behavior
 
@@ -622,10 +637,6 @@ namespace UE::AnimGen::Editor
 					FText::AsNumber(MaxCharacterNum),
 					FText::AsNumber(ActiveRanges.Num())));
 			}
-			else
-			{
-				ViewportClient->SetWarningMessage(FText());
-			}
 
 			// Adjust the timeline for the newly selected ranges
 
diff --git a/Engine/Plugins/Experimental/Animation/AnimGen/Source/AnimGenEditor/Private/AnimGenEditorControllerToolkit.cpp b/Engine/Plugins/Experimental/Animation/AnimGen/Source/AnimGenEditor/Private/AnimGenEditorControllerToolkit.cpp
--- a/Engine/Plugins/Experimental/Animation/AnimGen/Source/AnimGenEditor/Private/AnimGenEditorControllerToolkit.cpp
+++ b/Engine/Plugins/Experimental/Animation/AnimGen/Source/AnimGenEditor/Private/AnimGenEditorControllerToolkit.cpp
@@ -43,6 +43,7 @@
 #include "DetailCategoryBuilder.h"
 #include "DetailWidgetRow.h"
 #include "IDetailGroup.h"
+#include "SWarningOrErrorBox.h"
 #include "Widgets/Text/STextBlock.h"
 #include "Widgets/Input/SCheckBox.h"
 #include "Widgets/Input/SHyperlink.h"
@@ -185,8 +186,28 @@ namespace UE::AnimGen::Editor
 		// Data Category
 
 		IDetailCategoryBuilder& DataCategory = DetailBuilder.EditCategory("Data", LOCTEXT("DataCategory", "Data"), ECategoryPriority::Important);
-
-		// TODO: Display Errors when properties don't match.
+		DataCategory.AddProperty(GET_MEMBER_NAME_CHECKED(UAnimGenControllerTrainingSettings, Database), UAnimGenControllerTrainingSettings::StaticClass());
+		DataCategory.AddProperty(GET_MEMBER_NAME_CHECKED(UAnimGenControllerTrainingSettings, FrameRanges), UAnimGenControllerTrainingSettings::StaticClass());
+		DataCategory.AddProperty(GET_MEMBER_NAME_CHECKED(UAnimGenControllerTrainingSettings, AutoEncoder), UAnimGenControllerTrainingSettings::StaticClass());
+		DataCategory.AddProperty(GET_MEMBER_NAME_CHECKED(UAnimGenControllerTrainingSettings, Behavior), UAnimGenControllerTrainingSettings::StaticClass());
+
+		if (TrainingSettings->Behavior && (
+			!TrainingSettings->ControlSchema.IsValid() || 
+			!TrainingSettings->ControlSchema.IsElementValid(TrainingSettings->ControlSchemaElement)))
+		{
+			FDetailWidgetRow& NeedsReimportErrorRow = DataCategory.AddCustomRow(LOCTEXT("InvalidBehaviorWarning", "Error"))
+				.WholeRowContent()
+				[
+					SNew(SBox)
+						.Padding(FMargin(0.0f, 4.0f))
+						[
+							SNew(SWarningOrErrorBox)
+								.MessageStyle(EMessageStyle::Error)
+								.Message(LOCTEXT("InvalidBehaviorErrorMessage",
+									"Behavior is invalid."))
+						]
+				];
+		}
 
 		IDetailCategoryBuilder& ControlSetsCategory = DetailBuilder.EditCategory("ControlSets", LOCTEXT("ControlSetsCategory", "Control Sets"), ECategoryPriority::Important);
 
@@ -262,6 +283,7 @@ namespace UE::AnimGen::Editor
 					SNew(STextBlock).Text(
 						TrainingSettings->Database.IsNull() ? LOCTEXT("NoDatabase", "No Training Database Provided") :
 						!TrainingSettings->Database.Get() ? LOCTEXT("LoadingDatabase", "Loading Training Database...") :
+						!TrainingSettings->AutoEncoder.Get() ? LOCTEXT("LoadingAutoEncoder", "Loading Training AutoEncoder...") :
 						LOCTEXT("EmptyRanges", "No Control Sets returned by Behavior"))
 				];
 		}
@@ -379,6 +401,94 @@ namespace UE::AnimGen::Editor
 		return SNew(SFilterMultiColumnRow, OwnerTable).Item(Item);
 	}
 
+
+	TSharedRef<IDetailCustomization> FControllerDetails::MakeInstance()
+	{
+		return MakeShared<FControllerDetails>();
+	}
+
+	void FControllerDetails::CustomizeDetails(IDetailLayoutBuilder& DetailBuilder)
+	{
+		// Get the UAnimGenController object we are associated with
+
+		TArray<TWeakObjectPtr<UObject>> Objects;
+		DetailBuilder.GetObjectsBeingCustomized(Objects);
+		if (Objects.Num() != 1) { return; }
+
+		UAnimGenController* Controller = StaticCast<UAnimGenController*>(Objects[0].Pin().Get());
+		if (!Controller) { return; }
+
+		// Add Categories
+
+		IDetailCategoryBuilder& ControllerCategory = DetailBuilder.EditCategory("Controller", LOCTEXT("ControllerCategory", "Controller"), ECategoryPriority::Important);
+		ControllerCategory.AddProperty(GET_MEMBER_NAME_CHECKED(UAnimGenController, AutoEncoder), UAnimGenController::StaticClass());
+		ControllerCategory.AddProperty(GET_MEMBER_NAME_CHECKED(UAnimGenController, FrameRate), UAnimGenController::StaticClass());
+
+		// Check if Trained Database has been modified
+
+		if (Controller->IsValid() && Controller->TrainingSettings->Database.Get() && UAnimDatabaseFrameRangesLibrary::FrameRangesContentHash(
+			Controller->TrainingSettings->Database.Get(),
+			Controller->TrainingSettings->QueryRanges) != Controller->TrainedFrameRangesContentHash)
+		{
+			FDetailWidgetRow& NeedsReimportErrorRow = ControllerCategory.AddCustomRow(LOCTEXT("ModifiedWarning", "Warning"))
+				.WholeRowContent()
+				[
+					SNew(SBox)
+						.Padding(FMargin(0.0f, 4.0f))
+						[
+							SNew(SWarningOrErrorBox)
+								.MessageStyle(EMessageStyle::Warning)
+								.Message(LOCTEXT("DatabaseModifiedWarningMessage",
+									"Training Data appears to have changed since Controller was last trained."))
+						]
+				];
+		}
+
+		if (Controller->IsValid() && Controller->AutoEncoder && Controller->AutoEncoder->TrainedContentHash != Controller->TrainedAutoEncoderContentHash)
+		{
+			FDetailWidgetRow& NeedsReimportErrorRow = ControllerCategory.AddCustomRow(LOCTEXT("ModifiedWarning", "Warning"))
+				.WholeRowContent()
+				[
+					SNew(SBox)
+						.Padding(FMargin(0.0f, 4.0f))
+						[
+							SNew(SWarningOrErrorBox)
+								.MessageStyle(EMessageStyle::Warning)
+								.Message(LOCTEXT("AutoEncoderModifiedWarningMessage",
+									"Auto-Encoder appears to have changed since Controller was last trained. Controller may no longer function correctly."))
+						]
+				];
+		}
+
+		if (Controller->IsValid() && Controller->TrainingSettings->Behavior)
+		{
+			FAnimGenControlSchema ControllerSchema = UAnimGenControls::MakeControlSchema();
+			const FAnimGenControlSchemaElement SchemaElement = Controller->TrainingSettings->Behavior->SpecifyControl(ControllerSchema);
+
+			if (Learning::Observation::GetSchemaObjectsCompatibilityHash(ControllerSchema.ObservationSchema->Schema, SchemaElement.SchemaElement) != Controller->TrainedSchemaCompatibilityHash)
+			{
+				FDetailWidgetRow& NeedsReimportErrorRow = ControllerCategory.AddCustomRow(LOCTEXT("ModifiedWarning", "Warning"))
+					.WholeRowContent()
+					[
+						SNew(SBox)
+							.Padding(FMargin(0.0f, 4.0f))
+							[
+								SNew(SWarningOrErrorBox)
+									.MessageStyle(EMessageStyle::Warning)
+									.Message(LOCTEXT("BehaviorModifiedWarningMessage",
+										"Behavior appears to have changed since Controller was last trained."))
+							]
+					];
+			}
+		}
+
+		IDetailCategoryBuilder& StatisticsCategory = DetailBuilder.EditCategory("Statistics", LOCTEXT("Statistics", "Statistics"), ECategoryPriority::Important);
+		StatisticsCategory.AddProperty(GET_MEMBER_NAME_CHECKED(UAnimGenController, ControlEncoderSize), UAnimGenController::StaticClass());
+		StatisticsCategory.AddProperty(GET_MEMBER_NAME_CHECKED(UAnimGenController, DenoiserSize), UAnimGenController::StaticClass());
+		StatisticsCategory.AddProperty(GET_MEMBER_NAME_CHECKED(UAnimGenController, ControlEncoderInferenceTime), UAnimGenController::StaticClass());
+		StatisticsCategory.AddProperty(GET_MEMBER_NAME_CHECKED(UAnimGenController, DenoiserInferenceTime), UAnimGenController::StaticClass());
+	}
+
 	class SControllerViewportToolBar;
 
 	/** Custom Controller Viewport class. Contains a custom show menu at the top which contains some toggles for the ViewportSettings */
@@ -1620,6 +1730,7 @@ namespace UE::AnimGen::Editor
 	{
 		Controller = InController;
 		Controller->TrainingSettings->LoadDatabaseAsync();
+		Controller->TrainingSettings->LoadAutoEncoderAsync();
 		Controller->TestingSettings->LoadDatabaseAsync();
 
 		PreviewScene = MakeShared<AnimDatabase::Editor::FPreviewScene>(AnimDatabase::Editor::FPreviewScene::ConstructionValues().ShouldSimulatePhysics(true).ForceUseMovementComponentInNonGameWorld(true));
diff --git a/Engine/Plugins/Experimental/Animation/AnimGen/Source/AnimGenEditor/Public/AnimGenEditorControllerToolkit.h b/Engine/Plugins/Experimental/Animation/AnimGen/Source/AnimGenEditor/Public/AnimGenEditorControllerToolkit.h
--- a/Engine/Plugins/Experimental/Animation/AnimGen/Source/AnimGenEditor/Public/AnimGenEditorControllerToolkit.h
+++ b/Engine/Plugins/Experimental/Animation/AnimGen/Source/AnimGenEditor/Public/AnimGenEditorControllerToolkit.h
@@ -90,6 +90,16 @@ namespace UE::AnimGen::Editor
 		TSharedPtr<SListView<TSharedPtr<AnimDatabase::Editor::FQueryEntry>>> ListView;
 	};
 
+	/** Details customization for the controller */
+	class FControllerDetails : public IDetailCustomization
+	{
+	public:
+
+		static TSharedRef<IDetailCustomization> MakeInstance();
+
+		virtual void CustomizeDetails(IDetailLayoutBuilder& DetailBuilder) override;
+	};
+
 	enum class EControllerState : uint8
 	{
 		Stopped,
