AnimGen: Various small updates and fixes.

#Synced-CL 52443113

diff --git a/Engine/Plugins/Experimental/Animation/AnimDatabase/Source/AnimDatabase/Private/AnimDatabaseFrameAttribute.cpp b/Engine/Plugins/Experimental/Animation/AnimDatabase/Source/AnimDatabase/Private/AnimDatabaseFrameAttribute.cpp
--- a/Engine/Plugins/Experimental/Animation/AnimDatabase/Source/AnimDatabase/Private/AnimDatabaseFrameAttribute.cpp
+++ b/Engine/Plugins/Experimental/Animation/AnimDatabase/Source/AnimDatabase/Private/AnimDatabaseFrameAttribute.cpp
@@ -563,6 +563,7 @@ FAnimDatabaseFrameAttribute UAnimDatabaseFrameAttributeLibrary::MakeBoolFrameAtt
 	Out.Type = EAnimDatabaseAttributeType::Bool;
 	Out.FrameAttribute->FrameRangeSet = *FrameRanges.FrameRangeSet;
 	Out.FrameAttribute->AttributeData.SetNumUninitialized({ AttributeTypeSize(Out.Type), FrameRanges.FrameRangeSet->GetTotalFrameNum() });
+	UE::Learning::Array::Zero(Out.FrameAttribute->AttributeData);
 
 	UE::Learning::FrameRangeSet::ParallelForEachRange(Out.FrameAttribute->FrameRangeSet, [&Out, &Active](
 		const int32 TotalRangeIdx,
@@ -574,14 +575,14 @@ FAnimDatabaseFrameAttribute UAnimDatabaseFrameAttributeLibrary::MakeBoolFrameAtt
 			const int32 Sequence = Out.FrameAttribute->FrameRangeSet.GetEntrySequence(EntryIdx);
 			const int32 StartFrame = Out.FrameAttribute->FrameRangeSet.GetEntryRangeStart(EntryIdx, RangeIdx);
 
-			if (!Active.FrameRangeSet->ContainsSequence(Sequence))
-			{
-				UE::Learning::Array::Zero(Out.FrameAttribute->AttributeData[0].Slice(FrameOffset, FrameNum));
-			}
+			if (!Active.FrameRangeSet->IntersectsRange(Sequence, StartFrame, FrameNum)) { return; }
+
+			const int32 ActiveEntry = Active.FrameRangeSet->FindSequenceEntry(Sequence);
+			check(ActiveEntry != INDEX_NONE);
 
 			for (int32 FrameIdx = 0; FrameIdx < FrameNum; FrameIdx++)
 			{
-				Out.FrameAttribute->AttributeData[0][FrameOffset + FrameIdx] = Active.FrameRangeSet->Contains(Sequence, StartFrame + FrameIdx) ? 1.0f : 0.0f;
+				Out.FrameAttribute->AttributeData[0][FrameOffset + FrameIdx] = Active.FrameRangeSet->EntryContains(ActiveEntry, StartFrame + FrameIdx) ? 1.0f : 0.0f;
 			}
 
 		});
@@ -9014,13 +9015,16 @@ FAnimDatabaseFrameAttribute UAnimDatabaseFrameAttributeLibrary::FrameAttributeFi
 		return FAnimDatabaseFrameAttribute();
 	}
 
-	if (A.Type != EAnimDatabaseAttributeType::Float && A.Type != EAnimDatabaseAttributeType::Location && A.Type != EAnimDatabaseAttributeType::Direction)
+	if (A.Type != EAnimDatabaseAttributeType::Float && 
+		A.Type != EAnimDatabaseAttributeType::Bool &&
+		A.Type != EAnimDatabaseAttributeType::Location && 
+		A.Type != EAnimDatabaseAttributeType::Direction)
 	{
-		UE_LOGFMT(LogAnimDatabase, Error, "FrameAttributeFilterGaussian: FrameAttribute Types don't match. Got {Type}. Expected Float, Location, or Direction.", AttributeTypeNameInternal(A.Type));
+		UE_LOGFMT(LogAnimDatabase, Error, "FrameAttributeFilterGaussian: FrameAttribute Types don't match. Got {Type}. Expected Float, Bool, Location, or Direction.", AttributeTypeNameInternal(A.Type));
 		return FAnimDatabaseFrameAttribute();
 	}
 
-	if (A.Type == EAnimDatabaseAttributeType::Float)
+	if (A.Type == EAnimDatabaseAttributeType::Float || A.Type == EAnimDatabaseAttributeType::Bool)
 	{
 		FAnimDatabaseFrameAttribute Out = MakeEmptyFrameAttribute();
 		Out.Type = EAnimDatabaseAttributeType::Float;
@@ -10479,6 +10483,15 @@ FAnimDatabaseFrameAttribute UAnimDatabaseFrameAttributeFunction_BoneGlobalLinear
 	return UAnimDatabaseFrameAttributeLibrary::FrameAttributeBoneSpeed(Database, FrameRanges, BoneName);
 }
 
+FAnimDatabaseFrameAttribute UAnimDatabaseFrameAttributeFunction_GaussianSmoothingFilter::MakeFrameAttribute_Implementation(const UAnimDatabase* Database, const FAnimDatabaseFrameRanges& FrameRanges) const
+{
+	if (!Database) { return FAnimDatabaseFrameAttribute(); }
+
+	return UAnimDatabaseFrameAttributeLibrary::FrameAttributeFilterGaussian(
+		UAnimDatabaseFrameAttributeLibrary::MakeFrameAttributeFromFunction(Database, FrameRanges, FilterFrameAttribute),
+		FMath::RoundToInt(StandardDeviation * Database->GetFrameRate().AsDecimal()));
+}
+
 FAnimDatabaseFrameRanges UAnimDatabaseFrameRangesFunction_MajorityVoteFilter::MakeFrameRanges_Implementation(const UAnimDatabase* Database, const FAnimDatabaseFrameRanges& FrameRanges) const
 {
 	return UAnimDatabaseFrameAttributeLibrary::FrameAttributeBoolToFrameRanges(
diff --git a/Engine/Plugins/Experimental/Animation/AnimDatabase/Source/AnimDatabase/Private/AnimDatabaseFrameRanges.cpp b/Engine/Plugins/Experimental/Animation/AnimDatabase/Source/AnimDatabase/Private/AnimDatabaseFrameRanges.cpp
--- a/Engine/Plugins/Experimental/Animation/AnimDatabase/Source/AnimDatabase/Private/AnimDatabaseFrameRanges.cpp
+++ b/Engine/Plugins/Experimental/Animation/AnimDatabase/Source/AnimDatabase/Private/AnimDatabaseFrameRanges.cpp
@@ -307,6 +307,8 @@ FAnimDatabaseFrames UAnimDatabaseFramesLibrary::MakeFramesFromAnimNotifiesOnTrac
 {
 	TRACE_CPUPROFILER_EVENT_SCOPE(UAnimDatabaseFramesLibrary::MakeFramesFromAnimNotifiesOnTrack);
 
+#if WITH_EDITOR
+
 	if (!Database)
 	{
 		UE_LOGFMT(LogAnimDatabase, Error, "MakeFramesFromAnimNotifiesOnTrack: Database is nullptr.");
@@ -335,7 +337,6 @@ FAnimDatabaseFrames UAnimDatabaseFramesLibrary::MakeFramesFromAnimNotifiesOnTrac
 
 		int32 TrackIndex = INDEX_NONE;
 
-#if WITH_EDITOR
 		int32 TrackNum = Sequence->AnimNotifyTracks.Num();
 		for (int32 TrackIdx = 0; TrackIdx < TrackNum; TrackIdx++)
 		{
@@ -345,10 +346,6 @@ FAnimDatabaseFrames UAnimDatabaseFramesLibrary::MakeFramesFromAnimNotifiesOnTrac
 				break;
 			}
 		}
-#else
-		UE_LOGFMT(LogAnimDatabase, Error, "MakeFramesFromAnimNotifiesOnTrack: Track Names are only available in Editor.");
-		return FAnimDatabaseFrames();
-#endif
 
 		if (TrackIndex == INDEX_NONE) { continue; }
 
@@ -391,6 +388,13 @@ FAnimDatabaseFrames UAnimDatabaseFramesLibrary::MakeFramesFromAnimNotifiesOnTrac
 	}
 
 	return Out;
+
+#else
+
+	UE_LOGFMT(LogAnimDatabase, Error, "MakeFramesFromAnimNotifiesOnTrack: Track Names are only available in Editor.");
+	return FAnimDatabaseFrames();
+
+#endif
 }
 
 FAnimDatabaseFrames UAnimDatabaseFramesLibrary::MakeFramesFromSyncMarker(const UAnimDatabase* Database, const FAnimDatabaseFrameRanges& FrameRanges, const FName SyncMarkerName)
@@ -3150,6 +3154,14 @@ FAnimDatabaseFrameRanges UAnimDatabaseFrameRangesFunction_AnimNotifyStateInterse
 	return UAnimDatabaseFrameRangesLibrary::MakeFrameRangesFromAnimNotifyStatesArrayViewIntersection(Database, FrameRanges, AnimNotifyStates);
 }
 
+FAnimDatabaseFrameRanges UAnimDatabaseFrameRangesFunction_AnimNotifyStateWithMirrored::MakeFrameRanges_Implementation(const UAnimDatabase* Database, const FAnimDatabaseFrameRanges& FrameRanges) const
+{
+	return UAnimDatabaseFrameRangesLibrary::FrameRangesUnionFromArrayView({
+		UAnimDatabaseFrameRangesLibrary::MakeFrameRangesFromAnimNotifyState(Database, UAnimDatabaseFrameRangesLibrary::MakeFrameRangesFromNotMirrored(Database, FrameRanges), NotMirroredAnimNotifyState),
+		UAnimDatabaseFrameRangesLibrary::MakeFrameRangesFromAnimNotifyState(Database, UAnimDatabaseFrameRangesLibrary::MakeFrameRangesFromMirrored(Database, FrameRanges), MirroredAnimNotifyState)
+		});
+}
+
 FAnimDatabaseFrameRanges UAnimDatabaseFrameRangesFunction_AssetNameContains::MakeFrameRanges_Implementation(const UAnimDatabase* Database, const FAnimDatabaseFrameRanges& FrameRanges) const
 {
 	return UAnimDatabaseFrameRangesLibrary::MakeFrameRangesFromAssetNameContainsStringView(Database, FrameRanges, AssetNameSubstring);
@@ -3265,6 +3277,16 @@ FAnimDatabaseFrameRanges UAnimDatabaseFrameRangesFunction_IntersectsFrames::Make
 	return UAnimDatabaseFrameRangesLibrary::FrameRangesIntersectsFrames(FrameRanges, UAnimDatabaseFramesLibrary::MakeFramesFromFunction(Database, FrameRanges, Function));
 }
 
+FAnimDatabaseFrameRanges UAnimDatabaseFrameRangesFunction_Trim::MakeFrameRanges_Implementation(const UAnimDatabase* Database, const FAnimDatabaseFrameRanges& FrameRanges) const
+{
+	return UAnimDatabaseFrameRangesLibrary::FrameRangesTrim(UAnimDatabaseFrameRangesLibrary::MakeFrameRangesFromFunction(Database, FrameRanges, Function), TrimStart, TrimEnd);
+}
+
+FAnimDatabaseFrameRanges UAnimDatabaseFrameRangesFunction_Pad::MakeFrameRanges_Implementation(const UAnimDatabase* Database, const FAnimDatabaseFrameRanges& FrameRanges) const
+{
+	return UAnimDatabaseFrameRangesLibrary::FrameRangesPad(UAnimDatabaseFrameRangesLibrary::MakeFrameRangesFromFunction(Database, FrameRanges, Function), PadStart, PadEnd);
+}
+
 FAnimDatabaseFrameRanges UAnimDatabaseFrameRangesFunction_Interval::MakeFrameRanges_Implementation(const UAnimDatabase* Database, const FAnimDatabaseFrameRanges& FrameRanges) const
 {
 	return UAnimDatabaseFrameRangesLibrary::MakeFrameRangesFromFrames(UAnimDatabaseFrameRangesLibrary::MakeFramesFromFrameRangesInterval(FrameRanges, Interval));
diff --git a/Engine/Plugins/Experimental/Animation/AnimDatabase/Source/AnimDatabase/Private/AnimDatabaseMath.cpp b/Engine/Plugins/Experimental/Animation/AnimDatabase/Source/AnimDatabase/Private/AnimDatabaseMath.cpp
--- a/Engine/Plugins/Experimental/Animation/AnimDatabase/Source/AnimDatabase/Private/AnimDatabaseMath.cpp
+++ b/Engine/Plugins/Experimental/Animation/AnimDatabase/Source/AnimDatabase/Private/AnimDatabaseMath.cpp
@@ -3342,6 +3342,57 @@ namespace UE::AnimDatabase::Math
 		}
 	}
 
+	void ComputeMeanStdMasked(
+		const TLearningArrayView<1, float> OutMean,
+		const TLearningArrayView<1, float> OutStd,
+		const TLearningArrayView<2, const float> Data,
+		const TLearningArrayView<1, const bool> Mask,
+		const int32 ColSliceStart,
+		const int32 ColSliceNum)
+	{
+		TRACE_CPUPROFILER_EVENT_SCOPE(UE::AnimDatabase::Math::ComputeMeanStdMasked);
+
+		const int64 SampleNum = Data.Num<0>();
+		const int64 DimNum = Data.Num<1>();
+
+		check(OutMean.Num() == DimNum);
+		check(OutStd.Num() == DimNum);
+		check(ColSliceStart >= 0 && ColSliceNum >= 0 && ColSliceStart + ColSliceNum <= DimNum);
+		check(Mask.Num() == SampleNum);
+
+		Learning::Array::Zero(OutMean.Slice(ColSliceStart, ColSliceNum));
+		Learning::Array::Zero(OutStd.Slice(ColSliceStart, ColSliceNum));
+
+		int64 MaskedNum = 0;
+		for (int64 SampleIdx = 0; SampleIdx < SampleNum; SampleIdx++)
+		{
+			MaskedNum += Mask[SampleIdx] ? 1.0f : 0.0f;
+		}
+
+		if (MaskedNum == 0) { return; }
+
+		int64 MaskedIdx = 0;
+		for (int64 SampleIdx = 0; SampleIdx < SampleNum; SampleIdx++)
+		{
+			if (Mask[SampleIdx])
+			{
+				for (int64 DimIdx = ColSliceStart; DimIdx < ColSliceStart + ColSliceNum; DimIdx++)
+				{
+					OutStd[DimIdx] += (((float)MaskedIdx / MaskedNum) / (MaskedIdx + 1)) * FMath::Square(Data[SampleIdx][DimIdx] - OutMean[DimIdx]);
+					OutMean[DimIdx] += (Data[SampleIdx][DimIdx] - OutMean[DimIdx]) / (MaskedIdx + 1);
+				}
+				MaskedIdx++;
+			}
+		}
+
+		check(MaskedIdx == MaskedNum);
+
+		for (int64 DimIdx = ColSliceStart; DimIdx < ColSliceStart + ColSliceNum; DimIdx++)
+		{
+			OutStd[DimIdx] = FMath::Sqrt(OutStd[DimIdx]);
+		}
+	}
+
 	void ComputeMeanStd(
 		FVector3f& OutMean,
 		FVector3f& OutStd,
diff --git a/Engine/Plugins/Experimental/Animation/AnimDatabase/Source/AnimDatabase/Private/AnimDatabasePose.cpp b/Engine/Plugins/Experimental/Animation/AnimDatabase/Source/AnimDatabase/Private/AnimDatabasePose.cpp
--- a/Engine/Plugins/Experimental/Animation/AnimDatabase/Source/AnimDatabase/Private/AnimDatabasePose.cpp
+++ b/Engine/Plugins/Experimental/Animation/AnimDatabase/Source/AnimDatabase/Private/AnimDatabasePose.cpp
@@ -1334,7 +1334,7 @@ namespace UE::AnimDatabase
 
 						case EAnimDatabaseAttributeType::Bool:
 						{
-							OutPoseVectors[FrameIdx][Offset + 1] = (InPoseData.AttributeData.AttributeData[FrameIdx][AttributeOffet + 0] != 0.0f) ? 1.0f : -1.0f;
+							OutPoseVectors[FrameIdx][Offset + 1] = (InPoseData.AttributeData.AttributeData[FrameIdx][AttributeOffet + 0] > 0.0f) ? 1.0f : -1.0f;
 							break;
 						}
 
@@ -1844,7 +1844,7 @@ namespace UE::AnimDatabase
 			check(OutPoseVectorScale.Num() == InPoseVectors.Num<1>());
 			check(InPoseVectors.Num<1>() == PoseVectorSize(BoneLocationNum, BoneRotationNum, BoneScaleNum, AttributeTypes));
 
-			UE::Learning::SlicedParallelFor(InPoseVectors.Num<1>(), 64, [&OutPoseVectorOffset, &OutPoseVectorScale, &InPoseVectors](const int32 Start, const int32 Length)
+			UE::Learning::SlicedParallelFor(InPoseVectors.Num<1>(), 16, [&OutPoseVectorOffset, &OutPoseVectorScale, &InPoseVectors](const int32 Start, const int32 Length)
 			{
 				Math::ComputeMeanStd(
 					OutPoseVectorOffset, 
@@ -1853,72 +1853,79 @@ namespace UE::AnimDatabase
 					Start, Length);
 			});
 
+			// Average scales across different value types and max with small epsilon
+
+			const float Epsilon = UE_SMALL_NUMBER;
+
 			int32 Offset = 0;
 
 			float RootLogScaleStdMean = 0.0f;
 			Math::ComputeMean(RootLogScaleStdMean, OutPoseVectorScale.Slice(Offset, 3));
-			Learning::Array::Set(OutPoseVectorScale.Slice(Offset, 3), RootLogScaleStdMean + UE_KINDA_SMALL_NUMBER);
+			Learning::Array::Set(OutPoseVectorScale.Slice(Offset, 3), FMath::Max(RootLogScaleStdMean, Epsilon));
 			Offset += 3;
 
 			float LocalRootLinearVelocityStdMean = 0.0f;
 			Math::ComputeMean(LocalRootLinearVelocityStdMean, OutPoseVectorScale.Slice(Offset, 3));
-			Learning::Array::Set(OutPoseVectorScale.Slice(Offset, 3), LocalRootLinearVelocityStdMean + UE_KINDA_SMALL_NUMBER);
+			Learning::Array::Set(OutPoseVectorScale.Slice(Offset, 3), FMath::Max(LocalRootLinearVelocityStdMean, Epsilon));
 			Offset += 3;
 
 			float LocalRootAngularVelocityStdMean = 0.0f;
 			Math::ComputeMean(LocalRootAngularVelocityStdMean, OutPoseVectorScale.Slice(Offset, 3));
-			Learning::Array::Set(OutPoseVectorScale.Slice(Offset, 3), LocalRootAngularVelocityStdMean + UE_KINDA_SMALL_NUMBER);
+			Learning::Array::Set(OutPoseVectorScale.Slice(Offset, 3), FMath::Max(LocalRootAngularVelocityStdMean, Epsilon));
 			Offset += 3;
 
 			float RootScalarVelocityStdMean = 0.0f;
 			Math::ComputeMean(RootScalarVelocityStdMean, OutPoseVectorScale.Slice(Offset, 3));
-			Learning::Array::Set(OutPoseVectorScale.Slice(Offset, 3), RootScalarVelocityStdMean + UE_KINDA_SMALL_NUMBER);
+			Learning::Array::Set(OutPoseVectorScale.Slice(Offset, 3), FMath::Max(RootScalarVelocityStdMean, Epsilon));
 			Offset += 3;
 
 			float BoneLocationStdMean = 0.0f;
 			Math::ComputeMean(BoneLocationStdMean, OutPoseVectorScale.Slice(Offset, BoneLocationNum * 3));
-			Learning::Array::Set(OutPoseVectorScale.Slice(Offset, BoneLocationNum * 3), BoneLocationStdMean + UE_KINDA_SMALL_NUMBER);
+			Learning::Array::Set(OutPoseVectorScale.Slice(Offset, BoneLocationNum * 3), FMath::Max(BoneLocationStdMean, Epsilon));
 			Offset += BoneLocationNum * 3;
 
 			float BoneRotationStdMean = 0.0f;
 			Math::ComputeMean(BoneRotationStdMean, OutPoseVectorScale.Slice(Offset, BoneRotationNum * 6));
-			Learning::Array::Set(OutPoseVectorScale.Slice(Offset, BoneRotationNum * 6), BoneRotationStdMean + UE_KINDA_SMALL_NUMBER);
+			Learning::Array::Set(OutPoseVectorScale.Slice(Offset, BoneRotationNum * 6), FMath::Max(BoneRotationStdMean, Epsilon));
 			Offset += BoneRotationNum * 6;
 
 			float BoneScaleStdMean = 0.0f;
 			Math::ComputeMean(BoneScaleStdMean, OutPoseVectorScale.Slice(Offset, BoneScaleNum * 3));
-			Learning::Array::Set(OutPoseVectorScale.Slice(Offset, BoneScaleNum * 3), BoneScaleStdMean + UE_KINDA_SMALL_NUMBER);
+			Learning::Array::Set(OutPoseVectorScale.Slice(Offset, BoneScaleNum * 3), FMath::Max(BoneScaleStdMean, Epsilon));
 			Offset += BoneScaleNum * 3;
 
 			float BoneLinearVelocityStdMean = 0.0f;
 			Math::ComputeMean(BoneLinearVelocityStdMean, OutPoseVectorScale.Slice(Offset, BoneLocationNum * 3));
-			Learning::Array::Set(OutPoseVectorScale.Slice(Offset, BoneLocationNum * 3), BoneLinearVelocityStdMean + UE_KINDA_SMALL_NUMBER);
+			Learning::Array::Set(OutPoseVectorScale.Slice(Offset, BoneLocationNum * 3), FMath::Max(BoneLinearVelocityStdMean, Epsilon));
 			Offset += BoneLocationNum * 3;
 
 			float BoneAngularVelocityStdMean = 0.0f;
 			Math::ComputeMean(BoneAngularVelocityStdMean, OutPoseVectorScale.Slice(Offset, BoneRotationNum * 3));
-			Learning::Array::Set(OutPoseVectorScale.Slice(Offset, BoneRotationNum * 3), BoneAngularVelocityStdMean + UE_KINDA_SMALL_NUMBER);
+			Learning::Array::Set(OutPoseVectorScale.Slice(Offset, BoneRotationNum * 3), FMath::Max(BoneAngularVelocityStdMean, Epsilon));
 			Offset += BoneRotationNum * 3;
 
 			float BoneScalarVelocityStdMean = 0.0f;
 			Math::ComputeMean(BoneScalarVelocityStdMean, OutPoseVectorScale.Slice(Offset, BoneScaleNum * 3));
-			Learning::Array::Set(OutPoseVectorScale.Slice(Offset, BoneScaleNum * 3), BoneScalarVelocityStdMean + UE_KINDA_SMALL_NUMBER);
+			Learning::Array::Set(OutPoseVectorScale.Slice(Offset, BoneScaleNum * 3), FMath::Max(BoneScalarVelocityStdMean, Epsilon));
 			Offset += BoneScaleNum * 3;
 
-			const int32 AttributeNum = AttributeTypes.Num();
+			// Handle Attributes
 
-			// TODO: Compute over active values only
+			const int32 AttributeNum = AttributeTypes.Num();
 
 			for (int32 AttributeIdx = 0; AttributeIdx < AttributeNum; AttributeIdx++)
 			{
 				const EAnimDatabaseAttributeType AttributeType = AttributeTypes[AttributeIdx];
 				const int32 AttributeEncodingSize = AttributeTypeEncodingSize(AttributeType);
+				check(AttributeEncodingSize >= 1);
 
 				// Normalization for attribute active flag is always 0 and 1
 				OutPoseVectorOffset[Offset] = 0.0f;
 				OutPoseVectorScale[Offset] = 1.0f;
 				Offset += 1;
 
+				// Compute average across scales
+
 				switch (AttributeType)
 				{
 				case EAnimDatabaseAttributeType::Null:
@@ -1930,7 +1937,7 @@ namespace UE::AnimDatabase
 				{
 					check(AttributeEncodingSize == 1 + 1);
 
-					// Always normalize bools with just 0 and 1
+					// Always normalize bools with 1 and 0
 					Learning::Array::Set(OutPoseVectorOffset.Slice(Offset, 1), 0.0f);
 					Learning::Array::Set(OutPoseVectorScale.Slice(Offset, 1), 1.0f);
 					Offset += 1;
@@ -1950,7 +1957,7 @@ namespace UE::AnimDatabase
 					// Normalization for attribute value is done as normal
 					float BoneAttributeStdMean = 0.0f;
 					Math::ComputeMean(BoneAttributeStdMean, OutPoseVectorScale.Slice(Offset, AttributeEncodingSize - 1));
-					Learning::Array::Set(OutPoseVectorScale.Slice(Offset, AttributeEncodingSize - 1), BoneAttributeStdMean + UE_KINDA_SMALL_NUMBER);
+					Learning::Array::Set(OutPoseVectorScale.Slice(Offset, AttributeEncodingSize - 1), FMath::Max(BoneAttributeStdMean, Epsilon));
 					Offset += AttributeEncodingSize - 1;
 					break;
 				}
@@ -1962,17 +1969,17 @@ namespace UE::AnimDatabase
 
 					float BoneAttributeLocationStdMean = 0.0f;
 					Math::ComputeMean(BoneAttributeLocationStdMean, OutPoseVectorScale.Slice(Offset, 3));
-					Learning::Array::Set(OutPoseVectorScale.Slice(Offset, 3), BoneAttributeLocationStdMean + UE_KINDA_SMALL_NUMBER);
+					Learning::Array::Set(OutPoseVectorScale.Slice(Offset, 3), FMath::Max(BoneAttributeLocationStdMean, Epsilon));
 					Offset += 3;
 
 					float BoneAttributeRotationStdMean = 0.0f;
 					Math::ComputeMean(BoneAttributeRotationStdMean, OutPoseVectorScale.Slice(Offset, 6));
-					Learning::Array::Set(OutPoseVectorScale.Slice(Offset, 6), BoneAttributeRotationStdMean + UE_KINDA_SMALL_NUMBER);
+					Learning::Array::Set(OutPoseVectorScale.Slice(Offset, 6), FMath::Max(BoneAttributeRotationStdMean, Epsilon));
 					Offset += 6;
 
 					float BoneAttributeScaleStdMean = 0.0f;
 					Math::ComputeMean(BoneAttributeScaleStdMean, OutPoseVectorScale.Slice(Offset, 3));
-					Learning::Array::Set(OutPoseVectorScale.Slice(Offset, 3), BoneAttributeScaleStdMean + UE_KINDA_SMALL_NUMBER);
+					Learning::Array::Set(OutPoseVectorScale.Slice(Offset, 3), FMath::Max(BoneAttributeScaleStdMean, Epsilon));
 					Offset += 3;
 
 					break;
diff --git a/Engine/Plugins/Experimental/Animation/AnimDatabase/Source/AnimDatabase/Public/AnimDatabaseFrameAttribute.h b/Engine/Plugins/Experimental/Animation/AnimDatabase/Source/AnimDatabase/Public/AnimDatabaseFrameAttribute.h
--- a/Engine/Plugins/Experimental/Animation/AnimDatabase/Source/AnimDatabase/Public/AnimDatabaseFrameAttribute.h
+++ b/Engine/Plugins/Experimental/Animation/AnimDatabase/Source/AnimDatabase/Public/AnimDatabaseFrameAttribute.h
@@ -2371,6 +2371,27 @@ public:
 	UE_API virtual FAnimDatabaseFrameAttribute MakeFrameAttribute_Implementation(const UAnimDatabase* Database, const FAnimDatabaseFrameRanges& FrameRanges) const override;
 };
 
+/** Filter a frame attribute using a gaussian smoothing filter */
+UCLASS(BlueprintType, Blueprintable, meta = (DisplayName = "Gaussian Smoothing Filter"))
+class UAnimDatabaseFrameAttributeFunction_GaussianSmoothingFilter : public UAnimDatabaseFrameAttributeFunction
+{
+	GENERATED_BODY()
+
+public:
+
+	/** Frame attribute to filter */
+	UPROPERTY(EditAnywhere, Instanced, Category = "Settings")
+	TObjectPtr<UAnimDatabaseFrameAttributeFunction> FilterFrameAttribute;
+
+	/** Standard Deviation (in seconds) for the gaussian smoothing filter. */
+	UPROPERTY(EditAnywhere, Category = "Settings", meta=(ForceUnits="s"))
+	float StandardDeviation = 3.0f;
+
+public:
+
+	UE_API virtual FAnimDatabaseFrameAttribute MakeFrameAttribute_Implementation(const UAnimDatabase* Database, const FAnimDatabaseFrameRanges& FrameRanges) const override;
+};
+
 /** Filter frame ranges with a majority vote filter */
 UCLASS(BlueprintType, Blueprintable, meta = (DisplayName = "Majority Vote Filter"))
 class UAnimDatabaseFrameRangesFunction_MajorityVoteFilter : public UAnimDatabaseFrameRangesFunction
diff --git a/Engine/Plugins/Experimental/Animation/AnimDatabase/Source/AnimDatabase/Public/AnimDatabaseFrameRanges.h b/Engine/Plugins/Experimental/Animation/AnimDatabase/Source/AnimDatabase/Public/AnimDatabaseFrameRanges.h
--- a/Engine/Plugins/Experimental/Animation/AnimDatabase/Source/AnimDatabase/Public/AnimDatabaseFrameRanges.h
+++ b/Engine/Plugins/Experimental/Animation/AnimDatabase/Source/AnimDatabase/Public/AnimDatabaseFrameRanges.h
@@ -1051,6 +1051,25 @@ public:
 	UE_API virtual FAnimDatabaseFrameRanges MakeFrameRanges_Implementation(const UAnimDatabase* Database, const FAnimDatabaseFrameRanges& FrameRanges) const override;
 };
 
+/** Finds all the frame ranges with a given Anim Notify State classes for mirrored and un-mirrored sequences */
+UCLASS(BlueprintType, Blueprintable, meta = (DisplayName = "Anim Notify State with Mirrored"))
+class UAnimDatabaseFrameRangesFunction_AnimNotifyStateWithMirrored : public UAnimDatabaseFrameRangesFunction
+{
+	GENERATED_BODY()
+
+public:
+
+	UPROPERTY(EditAnywhere, Category = "Settings")
+	TSubclassOf<UAnimNotifyState> NotMirroredAnimNotifyState;
+
+	UPROPERTY(EditAnywhere, Category = "Settings")
+	TSubclassOf<UAnimNotifyState> MirroredAnimNotifyState;
+
+public:
+
+	UE_API virtual FAnimDatabaseFrameRanges MakeFrameRanges_Implementation(const UAnimDatabase* Database, const FAnimDatabaseFrameRanges& FrameRanges) const override;
+};
+
 /** Finds all the frame ranges with the given sub-string in the name */
 UCLASS(BlueprintType, Blueprintable, meta = (DisplayName = "Asset Name Contains"))
 class UAnimDatabaseFrameRangesFunction_AssetNameContains : public UAnimDatabaseFrameRangesFunction
@@ -1307,6 +1326,54 @@ public:
 	UE_API virtual FAnimDatabaseFrameRanges MakeFrameRanges_Implementation(const UAnimDatabase* Database, const FAnimDatabaseFrameRanges& FrameRanges) const override;
 };
 
+/** Trims the given frame ranges */
+UCLASS(BlueprintType, Blueprintable, meta = (DisplayName = "Trim"))
+class UAnimDatabaseFrameRangesFunction_Trim : public UAnimDatabaseFrameRangesFunction
+{
+	GENERATED_BODY()
+
+public:
+
+	UPROPERTY(EditAnywhere, Instanced, Category = "Settings")
+	TObjectPtr<UAnimDatabaseFrameRangesFunction> Function;
+
+	/** Number of frames to trim from the start of the frame ranges */
+	UPROPERTY(EditAnywhere, Category = "Settings", meta = (ClampMin = "0", UIMin = "0"))
+	int32 TrimStart = 0;
+
+	/** Number of frames to trim from the end of the frame ranges */
+	UPROPERTY(EditAnywhere, Category = "Settings", meta = (ClampMin = "0", UIMin = "0"))
+	int32 TrimEnd = 0;
+
+public:
+
+	UE_API virtual FAnimDatabaseFrameRanges MakeFrameRanges_Implementation(const UAnimDatabase* Database, const FAnimDatabaseFrameRanges& FrameRanges) const override;
+};
+
+/** Pads the given frame ranges */
+UCLASS(BlueprintType, Blueprintable, meta = (DisplayName = "Pad"))
+class UAnimDatabaseFrameRangesFunction_Pad : public UAnimDatabaseFrameRangesFunction
+{
+	GENERATED_BODY()
+
+public:
+
+	UPROPERTY(EditAnywhere, Instanced, Category = "Settings")
+	TObjectPtr<UAnimDatabaseFrameRangesFunction> Function;
+
+	/** Number of frames to pad the start of the frame ranges with */
+	UPROPERTY(EditAnywhere, Category = "Settings", meta = (ClampMin = "0", UIMin = "0"))
+	int32 PadStart = 0;
+
+	/** Number of frames to pad the end of the frame ranges with */
+	UPROPERTY(EditAnywhere, Category = "Settings", meta = (ClampMin = "0", UIMin = "0"))
+	int32 PadEnd = 0;
+
+public:
+
+	UE_API virtual FAnimDatabaseFrameRanges MakeFrameRanges_Implementation(const UAnimDatabase* Database, const FAnimDatabaseFrameRanges& FrameRanges) const override;
+};
+
 /** Entry specifying a frame ranges function */
 USTRUCT(BlueprintType)
 struct FAnimDatabaseFrameRangesEntry
diff --git a/Engine/Plugins/Experimental/Animation/AnimDatabase/Source/AnimDatabase/Public/AnimDatabaseMath.h b/Engine/Plugins/Experimental/Animation/AnimDatabase/Source/AnimDatabase/Public/AnimDatabaseMath.h
--- a/Engine/Plugins/Experimental/Animation/AnimDatabase/Source/AnimDatabase/Public/AnimDatabaseMath.h
+++ b/Engine/Plugins/Experimental/Animation/AnimDatabase/Source/AnimDatabase/Public/AnimDatabaseMath.h
@@ -1038,6 +1038,14 @@ namespace UE::AnimDatabase::Math
 		const int32 ColSliceStart,
 		const int32 ColSliceNum);
 
+	UE_API void ComputeMeanStdMasked(
+		const TLearningArrayView<1, float> OutMean,
+		const TLearningArrayView<1, float> OutStd,
+		const TLearningArrayView<2, const float> Data,
+		const TLearningArrayView<1, const bool> Mask,
+		const int32 ColSliceStart,
+		const int32 ColSliceNum);
+
 	UE_API void ComputeMeanStd(
 		FVector3f& OutMean,
 		FVector3f& OutStd,
diff --git a/Engine/Plugins/Experimental/Animation/AnimDatabase/Source/AnimDatabaseEditor/Private/AnimDatabaseEditorMode.cpp b/Engine/Plugins/Experimental/Animation/AnimDatabase/Source/AnimDatabaseEditor/Private/AnimDatabaseEditorMode.cpp
--- a/Engine/Plugins/Experimental/Animation/AnimDatabase/Source/AnimDatabaseEditor/Private/AnimDatabaseEditorMode.cpp
+++ b/Engine/Plugins/Experimental/Animation/AnimDatabase/Source/AnimDatabaseEditor/Private/AnimDatabaseEditorMode.cpp
@@ -84,11 +84,7 @@ namespace UE::AnimDatabase::Editor
 				InOutLinearVelocity = (FVector3f)RootAdjustmenTransform.TransformVectorNoScale((FVector)InOutLinearVelocity);
 				InOutAngularVelocity = (FVector3f)RootAdjustmenTransform.TransformVectorNoScale((FVector)InOutAngularVelocity);
 			}
-			else if (bDisableRootRotation)
-			{
-				checkNoEntry();
-			}
-			else
+			else 
 			{
 				InOutLocation = RootAdjustmenTransform.TransformPosition(InOutLocation);
 				InOutRotation = ((FQuat4f)RootAdjustmenTransform.GetRotation()) * InOutRotation;
diff --git a/Engine/Plugins/Experimental/Animation/AnimGen/Source/AnimGen/Private/AnimGenAutoEncoder.cpp b/Engine/Plugins/Experimental/Animation/AnimGen/Source/AnimGen/Private/AnimGenAutoEncoder.cpp
--- a/Engine/Plugins/Experimental/Animation/AnimGen/Source/AnimGen/Private/AnimGenAutoEncoder.cpp
+++ b/Engine/Plugins/Experimental/Animation/AnimGen/Source/AnimGen/Private/AnimGenAutoEncoder.cpp
@@ -211,56 +211,45 @@ bool UAnimGenAutoEncoderTrainingSettings::Update()
 			}
 		}
 
-		// Return if query update required
-
-		return bQueryUpdateRequired;
-	}
-
-	if (Database.IsNull() && (DatabaseContentHash != 0 || bForceRefresh))
-	{
-		bForceRefresh = false;
-		DatabaseContentHash = 0;
-		ActiveFrameRanges = nullptr;
-		QueryRanges = UAnimDatabaseFrameRangesLibrary::MakeEmptyFrameRanges();
-		UpdateQueryEntries();
-		ActiveDebugDrawer = nullptr;
-		SelectedRanges.Reset();
-		return true;
-	}
+		// Update Frame Attributes
 
-	return false;
-}
-
-bool UAnimGenAutoEncoderTrainingSettings::UpdateFrameAttributeObjects()
-{
-	if (const UAnimDatabase* DatabaseObject = Database.Get())
-	{
 		const int32 FrameAttributeNum = FrameAttributes.Num();
 
-		bool bUpdateRequired = ActiveFrameAttributes.Num() != FrameAttributeNum;
-
 		ActiveFrameAttributes.SetNum(FrameAttributeNum);
 		FrameAttributeObjects.SetNum(FrameAttributeNum);
 		FrameAttributeNames.SetNum(FrameAttributeNum);
 
 		for (int32 FrameAttributeIdx = 0; FrameAttributeIdx < FrameAttributeNum; FrameAttributeIdx++)
 		{
-			if (ActiveFrameAttributes[FrameAttributeIdx] != FrameAttributes[FrameAttributeIdx].FrameAttribute ||
+			if (bQueryUpdateRequired ||
+				ActiveFrameAttributes[FrameAttributeIdx] != FrameAttributes[FrameAttributeIdx].FrameAttribute ||
 				FrameAttributeNames[FrameAttributeIdx] != FrameAttributes[FrameAttributeIdx].Name)
 			{
-				bUpdateRequired = true;
+				bQueryUpdateRequired = true;
 				ActiveFrameAttributes[FrameAttributeIdx] = FrameAttributes[FrameAttributeIdx].FrameAttribute;
 				FrameAttributeNames[FrameAttributeIdx] = FrameAttributes[FrameAttributeIdx].Name;
 				FrameAttributeObjects[FrameAttributeIdx] = UAnimDatabaseFrameAttributeLibrary::MakeFrameAttributeFromFunction(DatabaseObject, QueryRanges, FrameAttributes[FrameAttributeIdx].FrameAttribute);
 			}
 		}
 
-		return bUpdateRequired;
+		// Return if query update required
+
+		return bQueryUpdateRequired;
 	}
-	else
+
+	if (Database.IsNull() && (DatabaseContentHash != 0 || bForceRefresh))
 	{
-		return false;
+		bForceRefresh = false;
+		DatabaseContentHash = 0;
+		ActiveFrameRanges = nullptr;
+		QueryRanges = UAnimDatabaseFrameRangesLibrary::MakeEmptyFrameRanges();
+		UpdateQueryEntries();
+		ActiveDebugDrawer = nullptr;
+		SelectedRanges.Reset();
+		return true;
 	}
+
+	return false;
 }
 
 #endif
diff --git a/Engine/Plugins/Experimental/Animation/AnimGen/Source/AnimGen/Public/AnimGenAutoEncoder.h b/Engine/Plugins/Experimental/Animation/AnimGen/Source/AnimGen/Public/AnimGenAutoEncoder.h
--- a/Engine/Plugins/Experimental/Animation/AnimGen/Source/AnimGen/Public/AnimGenAutoEncoder.h
+++ b/Engine/Plugins/Experimental/Animation/AnimGen/Source/AnimGen/Public/AnimGenAutoEncoder.h
@@ -361,14 +361,6 @@ private:
 	/** Internal function for updating the QueryEntries array from an updated QueryRanges */
 	UE_API void UpdateQueryEntries();
 
-public:
-
-	/**
-	 * Checks if the given FrameAttributes have been changed and if so updates the internal cached frames
-	 * and frame ranges objects as well as class default objects. Returns true if these were update and UI corresponding model may need updating
-	 */
-	UE_API bool UpdateFrameAttributeObjects();
-
 #endif
 
 #if WITH_EDITORONLY_DATA
diff --git a/Engine/Plugins/Experimental/Animation/AnimGen/Source/AnimGenEditor/Private/AnimGenEditorAutoEncoderMode.cpp b/Engine/Plugins/Experimental/Animation/AnimGen/Source/AnimGenEditor/Private/AnimGenEditorAutoEncoderMode.cpp
--- a/Engine/Plugins/Experimental/Animation/AnimGen/Source/AnimGenEditor/Private/AnimGenEditorAutoEncoderMode.cpp
+++ b/Engine/Plugins/Experimental/Animation/AnimGen/Source/AnimGenEditor/Private/AnimGenEditorAutoEncoderMode.cpp
@@ -84,10 +84,6 @@ namespace UE::AnimGen::Editor
 				InOutLinearVelocity = (FVector3f)RootAdjustmenTransform.TransformVectorNoScale((FVector)InOutLinearVelocity);
 				InOutAngularVelocity = (FVector3f)RootAdjustmenTransform.TransformVectorNoScale((FVector)InOutAngularVelocity);
 			}
-			else if (bDisableRootRotation)
-			{
-				checkNoEntry();
-			}
 			else
 			{
 				InOutLocation = RootAdjustmenTransform.TransformPosition(InOutLocation);
@@ -285,9 +281,7 @@ namespace UE::AnimGen::Editor
 
 		// If selection has changed or additional frames/frame ranges has changed update tracks
 
-		const bool bFrameAttributeObjectsUpdateRequired = AutoEncoder->TrainingSettings->UpdateFrameAttributeObjects();
-
-		if (bSelectedRangesUpdated || bFrameAttributeObjectsUpdateRequired)
+		if (bSelectedRangesUpdated)
 		{
 			EditorToolKit->TracksModel->ResetTracks();
 
diff --git a/Engine/Plugins/Experimental/Animation/AnimGen/Source/AnimGenEditor/Private/AnimGenEditorAutoEncoderToolkit.cpp b/Engine/Plugins/Experimental/Animation/AnimGen/Source/AnimGenEditor/Private/AnimGenEditorAutoEncoderToolkit.cpp
--- a/Engine/Plugins/Experimental/Animation/AnimGen/Source/AnimGenEditor/Private/AnimGenEditorAutoEncoderToolkit.cpp
+++ b/Engine/Plugins/Experimental/Animation/AnimGen/Source/AnimGenEditor/Private/AnimGenEditorAutoEncoderToolkit.cpp
@@ -1038,6 +1038,13 @@ namespace UE::AnimGen::Editor
 								PoseVectors.View,
 								Start,
 								Length);
+
+							for (int32 DimIdx = 0; DimIdx < Length; DimIdx++)
+							{
+								PoseVectorNormalizedStd[Start + DimIdx] = 
+									PoseVectorNormalizedStd[Start + DimIdx] > UE_SMALL_NUMBER ? 
+									FMath::Max(PoseVectorNormalizedStd[Start + DimIdx], UE_KINDA_SMALL_NUMBER) : 0.0f;
+							}
 						});
 				}
 
diff --git a/Engine/Plugins/Experimental/Animation/AnimGen/Source/AnimGenEditor/Private/AnimGenEditorControllerMode.cpp b/Engine/Plugins/Experimental/Animation/AnimGen/Source/AnimGenEditor/Private/AnimGenEditorControllerMode.cpp
--- a/Engine/Plugins/Experimental/Animation/AnimGen/Source/AnimGenEditor/Private/AnimGenEditorControllerMode.cpp
+++ b/Engine/Plugins/Experimental/Animation/AnimGen/Source/AnimGenEditor/Private/AnimGenEditorControllerMode.cpp
@@ -87,10 +87,6 @@ namespace UE::AnimGen::Editor
 				InOutLinearVelocity = (FVector3f)RootAdjustmenTransform.TransformVectorNoScale((FVector)InOutLinearVelocity);
 				InOutAngularVelocity = (FVector3f)RootAdjustmenTransform.TransformVectorNoScale((FVector)InOutAngularVelocity);
 			}
-			else if (bDisableRootRotation)
-			{
-				checkNoEntry();
-			}
 			else
 			{
 				InOutLocation = RootAdjustmenTransform.TransformPosition(InOutLocation);
