AnimDatabase: Fix for warning in Masked MeanStd

#Synced-CL 52443113

diff --git a/Engine/Plugins/Experimental/Animation/AnimDatabase/Source/AnimDatabase/Private/AnimDatabaseMath.cpp b/Engine/Plugins/Experimental/Animation/AnimDatabase/Source/AnimDatabase/Private/AnimDatabaseMath.cpp
--- a/Engine/Plugins/Experimental/Animation/AnimDatabase/Source/AnimDatabase/Private/AnimDatabaseMath.cpp
+++ b/Engine/Plugins/Experimental/Animation/AnimDatabase/Source/AnimDatabase/Private/AnimDatabaseMath.cpp
@@ -3366,7 +3366,7 @@ namespace UE::AnimDatabase::Math
 		int64 MaskedNum = 0;
 		for (int64 SampleIdx = 0; SampleIdx < SampleNum; SampleIdx++)
 		{
-			MaskedNum += Mask[SampleIdx] ? 1.0f : 0.0f;
+			MaskedNum += Mask[SampleIdx] ? 1 : 0;
 		}
 
 		if (MaskedNum == 0) { return; }
