AnimGen: Adjusted Display Names to avoid Naming Collisions

#Synced-CL 52479696

diff --git a/Engine/Plugins/Experimental/Animation/AnimDatabase/Source/AnimDatabase/Public/AnimDatabase.h b/Engine/Plugins/Experimental/Animation/AnimDatabase/Source/AnimDatabase/Public/AnimDatabase.h
--- a/Engine/Plugins/Experimental/Animation/AnimDatabase/Source/AnimDatabase/Public/AnimDatabase.h
+++ b/Engine/Plugins/Experimental/Animation/AnimDatabase/Source/AnimDatabase/Public/AnimDatabase.h
@@ -1077,7 +1077,7 @@ private:
 };
 
 /** Function for stripping unused bone-tracks from the AnimSequence to reduce the file size on desk  */
-UCLASS(BlueprintType, Blueprintable, meta = (DisplayName = "Strip Unused Bone Tracks"))
+UCLASS(BlueprintType, Blueprintable, meta = (DisplayName = "Strip Unused Bone Tracks Function"))
 class UAnimDatabaseFunction_StripUnusedBoneTracks : public UAnimDatabaseFunction
 {
 	GENERATED_BODY()
@@ -1094,7 +1094,7 @@ public:
 };
 
 /** Function for exporting the given ranges as BVH  */
-UCLASS(BlueprintType, Blueprintable, meta = (DisplayName = "Export to BVH"))
+UCLASS(BlueprintType, Blueprintable, meta = (DisplayName = "Export to BVH Function"))
 class UAnimDatabaseFunction_ExportBVH : public UAnimDatabaseFunction
 {
 	GENERATED_BODY()
@@ -1123,7 +1123,7 @@ public:
 };
 
 /** Function for importing from BVH  */
-UCLASS(BlueprintType, Blueprintable, meta = (DisplayName = "Import from BVH"))
+UCLASS(BlueprintType, Blueprintable, meta = (DisplayName = "Import from BVH Function"))
 class UAnimDatabaseFunction_ImportBVH : public UAnimDatabaseFunction
 {
 	GENERATED_BODY()
@@ -1160,7 +1160,7 @@ public:
 };
 
 /** Function for exporting all AnimNotifies and AnimNotifyStates from the given animations  */
-UCLASS(BlueprintType, Blueprintable, meta = (DisplayName = "Export Notifies"))
+UCLASS(BlueprintType, Blueprintable, meta = (DisplayName = "Export Notifies Function"))
 class UAnimDatabaseFunction_ExportNotifies : public UAnimDatabaseFunction
 {
 	GENERATED_BODY()
@@ -1185,7 +1185,7 @@ public:
 };
 
 /** Function for importing AnimNotifies and AnimNotifyStates onto the given animations  */
-UCLASS(BlueprintType, Blueprintable, meta = (DisplayName = "Import Notifies"))
+UCLASS(BlueprintType, Blueprintable, meta = (DisplayName = "Import Notifies Function"))
 class UAnimDatabaseFunction_ImportNotifies : public UAnimDatabaseFunction
 {
 	GENERATED_BODY()
@@ -1206,7 +1206,7 @@ public:
 };
 
 /** Function for removing AnimNotifies and AnimNotifyStates from the given animations  */
-UCLASS(BlueprintType, Blueprintable, meta = (DisplayName = "Remove Notifies"))
+UCLASS(BlueprintType, Blueprintable, meta = (DisplayName = "Remove Notifies Function"))
 class UAnimDatabaseFunction_RemoveNotifies : public UAnimDatabaseFunction
 {
 	GENERATED_BODY()
@@ -1221,7 +1221,7 @@ public:
 };
 
 /** Function for placing IK bones at the correct hand and foot locations  */
-UCLASS(BlueprintType, Blueprintable, meta = (DisplayName = "Update IK Bones"))
+UCLASS(BlueprintType, Blueprintable, meta = (DisplayName = "Update IK Bones Function"))
 class UAnimDatabaseFunction_UpdateIKBones : public UAnimDatabaseFunction
 {
 	GENERATED_BODY()
@@ -1274,7 +1274,7 @@ public:
 };
 
 /** Function for computing the root transform using the look-at direction  */
-UCLASS(BlueprintType, Blueprintable, meta = (DisplayName = "Set Root from Look-At"))
+UCLASS(BlueprintType, Blueprintable, meta = (DisplayName = "Set Root from Look-At Function"))
 class UAnimDatabaseFunction_SetRootFromLookAt : public UAnimDatabaseFunction
 {
 	GENERATED_BODY()
@@ -1353,7 +1353,7 @@ public:
 };
 
 /** Function for baking root motion into the pelvis  */
-UCLASS(BlueprintType, Blueprintable, meta = (DisplayName = "Bake Root Motion into Pelvis"))
+UCLASS(BlueprintType, Blueprintable, meta = (DisplayName = "Bake Root Motion into Pelvis Function"))
 class UAnimDatabaseFunction_BakeRootMotionIntoPelvis : public UAnimDatabaseFunction
 {
 	GENERATED_BODY()
@@ -1380,7 +1380,7 @@ public:
 };
 
 /** Function for making the root motion start at the origin for all the provided ranges  */
-UCLASS(BlueprintType, Blueprintable, meta = (DisplayName = "Make Root Start at Origin"))
+UCLASS(BlueprintType, Blueprintable, meta = (DisplayName = "Make Root Start at Origin Function"))
 class UAnimDatabaseFunction_MakeRootStartAtOrigin : public UAnimDatabaseFunction
 {
 	GENERATED_BODY()
@@ -1404,7 +1404,7 @@ public:
 };
 
 /** Function for exporting ranges to AnimSequences  */
-UCLASS(BlueprintType, Blueprintable, meta = (DisplayName = "Export to Anim Sequences"))
+UCLASS(BlueprintType, Blueprintable, meta = (DisplayName = "Export to Anim Sequences Function"))
 class UAnimDatabaseFunction_ExportToAnimSequences : public UAnimDatabaseFunction
 {
 	GENERATED_BODY()
@@ -1434,7 +1434,7 @@ public:
 };
 
 /** Function for exporting root motion from the given animations  */
-UCLASS(BlueprintType, Blueprintable, meta = (DisplayName = "Export Root Motion"))
+UCLASS(BlueprintType, Blueprintable, meta = (DisplayName = "Export Root Motion Function"))
 class UAnimDatabaseFunction_ExportRootMotion : public UAnimDatabaseFunction
 {
 	GENERATED_BODY()
@@ -1463,7 +1463,7 @@ public:
 };
 
 /** Function for making ranges looping  */
-UCLASS(BlueprintType, Blueprintable, meta = (DisplayName = "Make Looped"))
+UCLASS(BlueprintType, Blueprintable, meta = (DisplayName = "Make Looped Function"))
 class UAnimDatabaseFunction_MakeLooped : public UAnimDatabaseFunction
 {
 	GENERATED_BODY()
@@ -1493,7 +1493,7 @@ public:
 };
 
 /** Function for patching ranges  */
-UCLASS(BlueprintType, Blueprintable, meta = (DisplayName = "Patch Bones"))
+UCLASS(BlueprintType, Blueprintable, meta = (DisplayName = "Patch Bones Function"))
 class UAnimDatabaseFunction_PatchBones : public UAnimDatabaseFunction
 {
 	GENERATED_BODY()
@@ -1523,7 +1523,7 @@ public:
 };
 
 /** Function for setting to the reference pose  */
-UCLASS(BlueprintType, Blueprintable, meta = (DisplayName = "Set to Reference Pose"))
+UCLASS(BlueprintType, Blueprintable, meta = (DisplayName = "Set to Reference Pose Function"))
 class UAnimDatabaseFunction_SetToReferencePose : public UAnimDatabaseFunction
 {
 	GENERATED_BODY()
@@ -1550,7 +1550,7 @@ public:
 };
 
 /** Function for applying a transform to a bone  */
-UCLASS(BlueprintType, Blueprintable, meta = (DisplayName = "Bone Apply Transform"))
+UCLASS(BlueprintType, Blueprintable, meta = (DisplayName = "Bone Apply Transform Function"))
 class UAnimDatabaseFunction_BoneApplyTransform : public UAnimDatabaseFunction
 {
 	GENERATED_BODY()
@@ -1580,7 +1580,7 @@ public:
 };
 
 /** Function for adding AnimNotifyStates from a reference object  */
-UCLASS(BlueprintType, Blueprintable, meta = (DisplayName = "Add Anim Notify State"))
+UCLASS(BlueprintType, Blueprintable, meta = (DisplayName = "Add Anim Notify State Function"))
 class UAnimDatabaseFunction_AddAnimNotifyState : public UAnimDatabaseFunction
 {
 	GENERATED_BODY()
@@ -1610,7 +1610,7 @@ public:
 };
 
 /** Function that can execute a sequence of sub-functions  */
-UCLASS(BlueprintType, Blueprintable, meta = (DisplayName = "Sequence"))
+UCLASS(BlueprintType, Blueprintable, meta = (DisplayName = "Sequence Function"))
 class UAnimDatabaseFunction_Sequence : public UAnimDatabaseFunction
 {
 	GENERATED_BODY()
@@ -1634,7 +1634,7 @@ public:
 };
 
 /** Function for adding an AnimNotify from a frame function  */
-UCLASS(BlueprintType, Blueprintable, meta = (DisplayName = "Add Anim Notify at Frames"))
+UCLASS(BlueprintType, Blueprintable, meta = (DisplayName = "Add Anim Notify at Frames Function"))
 class UAnimDatabaseFunction_AddAnimNotifyAtFrames : public UAnimDatabaseFunction
 {
 	GENERATED_BODY()
@@ -1667,7 +1667,7 @@ public:
 };
 
 /** Function for adding an AnimNotifyState from a frame ranges function  */
-UCLASS(BlueprintType, Blueprintable, meta = (DisplayName = "Add Anim Notify State at Frame Ranges"))
+UCLASS(BlueprintType, Blueprintable, meta = (DisplayName = "Add Anim Notify State at Frame Ranges Function"))
 class UAnimDatabaseFunction_AddAnimNotifyStateAtFrameRanges : public UAnimDatabaseFunction
 {
 	GENERATED_BODY()
@@ -1700,7 +1700,7 @@ public:
 };
 
 /** Function for adding a curve from a frame attribute function  */
-UCLASS(BlueprintType, Blueprintable, meta = (DisplayName = "Add Curve from Frame Attribute"))
+UCLASS(BlueprintType, Blueprintable, meta = (DisplayName = "Add Curve from Frame Attribute Function"))
 class UAnimDatabaseFunction_AddCurveFromFrameAttribute : public UAnimDatabaseFunction
 {
 	GENERATED_BODY()
@@ -1733,7 +1733,7 @@ public:
 };
 
 /** Function for processing contacts */
-UCLASS(BlueprintType, Blueprintable, DontCollapseCategories, meta = (DisplayName = "Process Contacts"))
+UCLASS(BlueprintType, Blueprintable, DontCollapseCategories, meta = (DisplayName = "Process Contacts Function"))
 class UAnimDatabaseFunction_ProcessContacts : public UAnimDatabaseFunction
 {
 	GENERATED_BODY()
@@ -1914,7 +1914,7 @@ public:
 
 
 /** Function for generating phase curves from another source */
-UCLASS(BlueprintType, Blueprintable, DontCollapseCategories, meta = (DisplayName = "Process Phase Curves"))
+UCLASS(BlueprintType, Blueprintable, DontCollapseCategories, meta = (DisplayName = "Process Phase Curves Function"))
 class UAnimDatabaseFunction_ProcessPhaseCurves : public UAnimDatabaseFunction
 {
 	GENERATED_BODY()
@@ -2006,7 +2006,7 @@ public:
 
 
 /** Function for setting a bone local transform from frame attribute function  */
-UCLASS(BlueprintType, Blueprintable, meta = (DisplayName = "Set Bone Local Transform from Frame Attribute"))
+UCLASS(BlueprintType, Blueprintable, meta = (DisplayName = "Set Bone Local Transform from Frame Attribute Function"))
 class UAnimDatabaseFunction_SetBoneLocalTransformFromFrameAttribute : public UAnimDatabaseFunction
 {
 	GENERATED_BODY()
@@ -2033,7 +2033,7 @@ public:
 };
 
 /** Function for setting a bone global transform from frame attribute function  */
-UCLASS(BlueprintType, Blueprintable, meta = (DisplayName = "Set Bone Global Transform from Frame Attribute"))
+UCLASS(BlueprintType, Blueprintable, meta = (DisplayName = "Set Bone Global Transform from Frame Attribute Function"))
 class UAnimDatabaseFunction_SetBoneGlobalTransformFromFrameAttribute : public UAnimDatabaseFunction
 {
 	GENERATED_BODY()
@@ -2069,7 +2069,7 @@ enum class EAnimDatabaseBoneSpace : uint8
 };
 
 /** Copy bone transforms from certain bones to others  */
-UCLASS(BlueprintType, Blueprintable, meta = (DisplayName = "Copy Bones"))
+UCLASS(BlueprintType, Blueprintable, meta = (DisplayName = "Copy Bones Function"))
 class UAnimDatabaseFunction_CopyBones : public UAnimDatabaseFunction
 {
 	GENERATED_BODY()
@@ -2098,7 +2098,7 @@ public:
 };
 
 /** Reset bone transforms to the reference pose  */
-UCLASS(BlueprintType, Blueprintable, meta = (DisplayName = "Reset Bones"))
+UCLASS(BlueprintType, Blueprintable, meta = (DisplayName = "Reset Bones Function"))
 class UAnimDatabaseFunction_ResetBones : public UAnimDatabaseFunction
 {
 	GENERATED_BODY()
@@ -2122,7 +2122,7 @@ public:
 };
 
 /** Function to set bone to root location at frame  */
-UCLASS(BlueprintType, Blueprintable, meta = (DisplayName = "Set Bone to Root Transform at Frames"))
+UCLASS(BlueprintType, Blueprintable, meta = (DisplayName = "Set Bone to Root Transform at Frames Function"))
 class UAnimDatabaseFunction_SetBoneToRootTransformAtFrames : public UAnimDatabaseFunction
 {
 	GENERATED_BODY()
@@ -2149,7 +2149,7 @@ public:
 };
 
 /** Function to adjust neck rotation to remove potential artifacts after retargeting */
-UCLASS(BlueprintType, Blueprintable, meta = (DisplayName = "Adjust Neck Rotation"))
+UCLASS(BlueprintType, Blueprintable, meta = (DisplayName = "Adjust Neck Rotation Function"))
 class UAnimDatabaseFunction_AdjustNeckRotation : public UAnimDatabaseFunction
 {
 	GENERATED_BODY()
@@ -2198,7 +2198,7 @@ public:
 };
 
 /** Function to remove foot ground penetration */
-UCLASS(BlueprintType, Blueprintable, meta = (DisplayName = "Remove Foot Ground Penetration"))
+UCLASS(BlueprintType, Blueprintable, meta = (DisplayName = "Remove Foot Ground Penetration Function"))
 class UAnimDatabaseFunction_RemoveFootGroundPenetration : public UAnimDatabaseFunction
 {
 	GENERATED_BODY()
@@ -2251,7 +2251,7 @@ public:
 };
 
 /** Computes various statistics about the provided ranges */
-UCLASS(BlueprintType, Blueprintable, meta = (DisplayName = "Statistics"))
+UCLASS(BlueprintType, Blueprintable, meta = (DisplayName = "Statistics Function"))
 class UAnimDatabaseFunction_Statistics : public UAnimDatabaseFunction
 {
 	GENERATED_BODY()
@@ -2276,7 +2276,7 @@ public:
 };
 
 /** Computes various stats about root speeds for the provide frame ranges */
-UCLASS(BlueprintType, Blueprintable, meta = (DisplayName = "Speed Statistics"))
+UCLASS(BlueprintType, Blueprintable, meta = (DisplayName = "Speed Statistics Function"))
 class UAnimDatabaseFunction_SpeedStatistics : public UAnimDatabaseFunction
 {
 	GENERATED_BODY()
@@ -2327,7 +2327,7 @@ public:
 };
 
 /** Computes various statistics regarding transitions given by the provided FrameRanges  */
-UCLASS(BlueprintType, Blueprintable, meta = (DisplayName = "Transition Statistics"))
+UCLASS(BlueprintType, Blueprintable, meta = (DisplayName = "Transition Statistics Function"))
 class UAnimDatabaseFunction_TransitionStatistics : public UAnimDatabaseFunction
 {
 	GENERATED_BODY()
@@ -2427,7 +2427,7 @@ public:
 };
 
 /** Debug Draw class for drawing contact curves */
-UCLASS(BlueprintType, Blueprintable, meta = (DisplayName = "Contact Curves"))
+UCLASS(BlueprintType, Blueprintable, meta = (DisplayName = "Contact Curves Debug Draw"))
 class UAnimDatabaseDebugDraw_ContactCurves : public UAnimDatabaseDebugDraw
 {
 	GENERATED_BODY()
@@ -2489,7 +2489,7 @@ public:
 
 
 /** Debug Draw class for multiple other debug drawers */
-UCLASS(BlueprintType, Blueprintable, meta = (DisplayName = "Multi"))
+UCLASS(BlueprintType, Blueprintable, meta = (DisplayName = "Multi Debug Draw"))
 class UAnimDatabaseDebugDraw_Multi : public UAnimDatabaseDebugDraw
 {
 	GENERATED_BODY()
@@ -2526,7 +2526,7 @@ public:
 };
 
 /** Debug Draw class for drawing root orientation */
-UCLASS(BlueprintType, Blueprintable, meta = (DisplayName = "Root Orientation"))
+UCLASS(BlueprintType, Blueprintable, meta = (DisplayName = "Root Orientation Debug Draw"))
 class UAnimDatabaseDebugDraw_RootOrientation : public UAnimDatabaseDebugDraw
 {
 	GENERATED_BODY()
@@ -2569,7 +2569,7 @@ public:
 };
 
 /** Debug Draw class for drawing the current movement direction */
-UCLASS(BlueprintType, Blueprintable, meta = (DisplayName = "Movement Direction"))
+UCLASS(BlueprintType, Blueprintable, meta = (DisplayName = "Movement Direction Debug Draw"))
 class UAnimDatabaseDebugDraw_MovementDirection : public UAnimDatabaseDebugDraw
 {
 	GENERATED_BODY()
@@ -2624,7 +2624,7 @@ public:
 };
 
 /** Debug Draw the transform of the given bone */
-UCLASS(BlueprintType, Blueprintable, meta = (DisplayName = "Bone Transform"))
+UCLASS(BlueprintType, Blueprintable, meta = (DisplayName = "Bone Transform Debug Draw"))
 class UAnimDatabaseDebugDraw_BoneTransform : public UAnimDatabaseDebugDraw
 {
 	GENERATED_BODY()
@@ -2672,7 +2672,7 @@ public:
 
 
 /** Debug Draw a bone attachment such as a prop */
-UCLASS(BlueprintType, Blueprintable, meta = (DisplayName = "Bone Attachment"))
+UCLASS(BlueprintType, Blueprintable, meta = (DisplayName = "Bone Attachment Debug Draw"))
 class UAnimDatabaseDebugDraw_BoneAttachment : public UAnimDatabaseDebugDraw
 {
 	GENERATED_BODY()
@@ -2735,7 +2735,7 @@ public:
 };
 
 /** Debug Draw class for drawing root velocities */
-UCLASS(BlueprintType, Blueprintable, meta = (DisplayName = "Root Velocity"))
+UCLASS(BlueprintType, Blueprintable, meta = (DisplayName = "Root Velocity Debug Draw"))
 class UAnimDatabaseDebugDraw_RootVelocity : public UAnimDatabaseDebugDraw
 {
 	GENERATED_BODY()
@@ -2785,7 +2785,7 @@ public:
 
 
 /** Debug Draw class for drawing a chair at the attach bone */
-UCLASS(BlueprintType, Blueprintable, meta = (DisplayName = "Chair on Attach"))
+UCLASS(BlueprintType, Blueprintable, meta = (DisplayName = "Chair on Attach Debug Draw"))
 class UAnimDatabaseDebugDraw_ChairOnAttach : public UAnimDatabaseDebugDraw
 {
 	GENERATED_BODY()
diff --git a/Engine/Plugins/Experimental/Animation/AnimDatabase/Source/AnimDatabase/Public/AnimDatabaseFrameAttribute.h b/Engine/Plugins/Experimental/Animation/AnimDatabase/Source/AnimDatabase/Public/AnimDatabaseFrameAttribute.h
--- a/Engine/Plugins/Experimental/Animation/AnimDatabase/Source/AnimDatabase/Public/AnimDatabaseFrameAttribute.h
+++ b/Engine/Plugins/Experimental/Animation/AnimDatabase/Source/AnimDatabase/Public/AnimDatabaseFrameAttribute.h
@@ -2007,7 +2007,7 @@ public:
 };
 
 /** Computes a bool Frame Attribute for the given Frame Ranges  */
-UCLASS(BlueprintType, Blueprintable, meta = (DisplayName = "Frame Ranges"))
+UCLASS(BlueprintType, Blueprintable, meta = (DisplayName = "Frame Ranges Frame Attribute"))
 class UAnimDatabaseFrameAttributeFunction_FrameRanges : public UAnimDatabaseFrameAttributeFunction
 {
 	GENERATED_BODY()
@@ -2024,7 +2024,7 @@ public:
 };
 
 /** Computes an Attribute for the given Frame Ranges  */
-UCLASS(BlueprintType, Blueprintable, meta = (DisplayName = "Filter"))
+UCLASS(BlueprintType, Blueprintable, meta = (DisplayName = "Filter Frame Attribute"))
 class UAnimDatabaseFrameAttributeFunction_Filter : public UAnimDatabaseFrameAttributeFunction
 {
 	GENERATED_BODY()
@@ -2045,7 +2045,7 @@ public:
 };
 
 /** Finds all the potential outlier frame ranges in a database */
-UCLASS(BlueprintType, Blueprintable, meta = (DisplayName = "Outliers"))
+UCLASS(BlueprintType, Blueprintable, meta = (DisplayName = "Outliers Frame Attribute"))
 class UAnimDatabaseFrameRangesFunction_Outliers : public UAnimDatabaseFrameRangesFunction
 {
 	GENERATED_BODY()
@@ -2094,7 +2094,7 @@ public:
 };
 
 /** Computes a contact curve attribute */
-UCLASS(BlueprintType, Blueprintable, meta = (DisplayName = "Contact Curve"))
+UCLASS(BlueprintType, Blueprintable, meta = (DisplayName = "Contact Curve Frame Attribute"))
 class UAnimDatabaseFrameAttributeFunction_ContactCurve : public UAnimDatabaseFrameAttributeFunction
 {
 	GENERATED_BODY()
@@ -2138,7 +2138,7 @@ public:
 };
 
 /** Computes contact frames */
-UCLASS(BlueprintType, Blueprintable, meta = (DisplayName = "Contact Frames"))
+UCLASS(BlueprintType, Blueprintable, meta = (DisplayName = "Contact Frames Frame Attribute"))
 class UAnimDatabaseFramesFunction_ContactFrames : public UAnimDatabaseFramesFunction
 {
 	GENERATED_BODY()
@@ -2182,7 +2182,7 @@ public:
 };
 
 /** Computes a contact event attribute */
-UCLASS(BlueprintType, Blueprintable, meta = (DisplayName = "Contact Event"))
+UCLASS(BlueprintType, Blueprintable, meta = (DisplayName = "Contact Event Frame Attribute"))
 class UAnimDatabaseFrameAttributeFunction_ContactEvent : public UAnimDatabaseFrameAttributeFunction
 {
 	GENERATED_BODY()
@@ -2226,7 +2226,7 @@ public:
 };
 
 /** Fits a transform to the given locations and directions */
-UCLASS(BlueprintType, Blueprintable, meta = (DisplayName = "Fit Transform"))
+UCLASS(BlueprintType, Blueprintable, meta = (DisplayName = "Fit Transform Frame Attribute"))
 class UAnimDatabaseFrameAttributeFunction_FitTransform : public UAnimDatabaseFrameAttributeFunction
 {
 	GENERATED_BODY()
@@ -2271,7 +2271,7 @@ public:
 };
 
 /** Computes a bone global transform attribute */
-UCLASS(BlueprintType, Blueprintable, meta = (DisplayName = "Bone Global Transform"))
+UCLASS(BlueprintType, Blueprintable, meta = (DisplayName = "Bone Global Transform Frame Attribute"))
 class UAnimDatabaseFrameAttributeFunction_BoneGlobalTransform : public UAnimDatabaseFrameAttributeFunction
 {
 	GENERATED_BODY()
@@ -2292,7 +2292,7 @@ public:
 };
 
 /** Computes a bone global transform attribute relative to another  */
-UCLASS(BlueprintType, Blueprintable, meta = (DisplayName = "Bone Relative Transform"))
+UCLASS(BlueprintType, Blueprintable, meta = (DisplayName = "Bone Relative Transform Frame Attribute"))
 class UAnimDatabaseFrameAttributeFunction_BoneRelativeTransform : public UAnimDatabaseFrameAttributeFunction
 {
 	GENERATED_BODY()
@@ -2313,7 +2313,7 @@ public:
 };
 
 /** Computes a bone global location attribute */
-UCLASS(BlueprintType, Blueprintable, meta = (DisplayName = "Bone Global Location"))
+UCLASS(BlueprintType, Blueprintable, meta = (DisplayName = "Bone Global Location Frame Attribute"))
 class UAnimDatabaseFrameAttributeFunction_BoneGlobalLocation : public UAnimDatabaseFrameAttributeFunction
 {
 	GENERATED_BODY()
@@ -2334,7 +2334,7 @@ public:
 };
 
 /** Computes a bone global direction attribute */
-UCLASS(BlueprintType, Blueprintable, meta = (DisplayName = "Bone Global Direction"))
+UCLASS(BlueprintType, Blueprintable, meta = (DisplayName = "Bone Global Direction Frame Attribute"))
 class UAnimDatabaseFrameAttributeFunction_BoneGlobalDirection : public UAnimDatabaseFrameAttributeFunction
 {
 	GENERATED_BODY()
@@ -2355,7 +2355,7 @@ public:
 };
 
 /** Computes a bone velocity magnitude attribute */
-UCLASS(BlueprintType, Blueprintable, meta = (DisplayName = "Bone Global Linear Velocity Magnitude"))
+UCLASS(BlueprintType, Blueprintable, meta = (DisplayName = "Bone Global Linear Velocity Magnitude Frame Attribute"))
 class UAnimDatabaseFrameAttributeFunction_BoneGlobalLinearVelocityMagnitude : public UAnimDatabaseFrameAttributeFunction
 {
 	GENERATED_BODY()
@@ -2372,7 +2372,7 @@ public:
 };
 
 /** Filter a frame attribute using a gaussian smoothing filter */
-UCLASS(BlueprintType, Blueprintable, meta = (DisplayName = "Gaussian Smoothing Filter"))
+UCLASS(BlueprintType, Blueprintable, meta = (DisplayName = "Gaussian Smoothing Filter Frame Attribute"))
 class UAnimDatabaseFrameAttributeFunction_GaussianSmoothingFilter : public UAnimDatabaseFrameAttributeFunction
 {
 	GENERATED_BODY()
@@ -2393,7 +2393,7 @@ public:
 };
 
 /** Filter frame ranges with a majority vote filter */
-UCLASS(BlueprintType, Blueprintable, meta = (DisplayName = "Majority Vote Filter"))
+UCLASS(BlueprintType, Blueprintable, meta = (DisplayName = "Majority Vote Filter Frame Attribute"))
 class UAnimDatabaseFrameRangesFunction_MajorityVoteFilter : public UAnimDatabaseFrameRangesFunction
 {
 	GENERATED_BODY()
@@ -2416,7 +2416,7 @@ public:
 };
 
 /** Frame ranges where the character is moving in the given direction */
-UCLASS(BlueprintType, Blueprintable, meta = (DisplayName = "Moving in Direction"))
+UCLASS(BlueprintType, Blueprintable, meta = (DisplayName = "Moving in Direction Frame Attribute"))
 class UAnimDatabaseFrameRangesFunction_MovingInDirection : public UAnimDatabaseFrameRangesFunction
 {
 	GENERATED_BODY()
@@ -2441,7 +2441,7 @@ public:
 };
 
 /** Frame ranges where the root linear velocity magnitude is below some threshold */
-UCLASS(BlueprintType, Blueprintable, meta = (DisplayName = "Root Linear Velocity Below Threshold"))
+UCLASS(BlueprintType, Blueprintable, meta = (DisplayName = "Root Linear Velocity Below Threshold Frame Attribute"))
 class UAnimDatabaseFrameRangesFunction_RootLinearVelocityBelowThreshold : public UAnimDatabaseFrameRangesFunction
 {
 	GENERATED_BODY()
@@ -2457,7 +2457,7 @@ public:
 };
 
 /** Frame ranges where the root linear velocity magnitude is above some threshold */
-UCLASS(BlueprintType, Blueprintable, meta = (DisplayName = "Root Linear Velocity Above Threshold"))
+UCLASS(BlueprintType, Blueprintable, meta = (DisplayName = "Root Linear Velocity Above Threshold Frame Attribute"))
 class UAnimDatabaseFrameRangesFunction_RootLinearVelocityAboveThreshold : public UAnimDatabaseFrameRangesFunction
 {
 	GENERATED_BODY()
diff --git a/Engine/Plugins/Experimental/Animation/AnimDatabase/Source/AnimDatabase/Public/AnimDatabaseFrameRanges.h b/Engine/Plugins/Experimental/Animation/AnimDatabase/Source/AnimDatabase/Public/AnimDatabaseFrameRanges.h
--- a/Engine/Plugins/Experimental/Animation/AnimDatabase/Source/AnimDatabase/Public/AnimDatabaseFrameRanges.h
+++ b/Engine/Plugins/Experimental/Animation/AnimDatabase/Source/AnimDatabase/Public/AnimDatabaseFrameRanges.h
@@ -233,7 +233,7 @@ public:
 };
 
 /** Returns the empty set of Frames */
-UCLASS(BlueprintType, Blueprintable, meta = (DisplayName = "Empty"))
+UCLASS(BlueprintType, Blueprintable, meta = (DisplayName = "Empty Frames"))
 class UAnimDatabaseFramesFunction_Empty : public UAnimDatabaseFramesFunction
 {
 	GENERATED_BODY()
@@ -244,7 +244,7 @@ public:
 };
 
 /** Finds all the frames for a given Anim Notify Class */
-UCLASS(BlueprintType, Blueprintable, meta = (DisplayName = "Anim Notify"))
+UCLASS(BlueprintType, Blueprintable, meta = (DisplayName = "Anim Notify Frames"))
 class UAnimDatabaseFramesFunction_AnimNotify : public UAnimDatabaseFramesFunction
 {
 	GENERATED_BODY()
@@ -260,7 +260,7 @@ public:
 };
 
 /** Finds all the frames for all Anim Notifies on a given track */
-UCLASS(BlueprintType, Blueprintable, meta = (DisplayName = "Anim Notifies on Track"))
+UCLASS(BlueprintType, Blueprintable, meta = (DisplayName = "Anim Notifies on Track Frames"))
 class UAnimDatabaseFramesFunction_AnimNotifiesOnTrack : public UAnimDatabaseFramesFunction
 {
 	GENERATED_BODY()
@@ -276,7 +276,7 @@ public:
 };
 
 /** Finds all the frames for a given Sync Marker */
-UCLASS(BlueprintType, Blueprintable, meta = (DisplayName = "Sync Marker"))
+UCLASS(BlueprintType, Blueprintable, meta = (DisplayName = "Sync Marker Frames"))
 class UAnimDatabaseFramesFunction_SyncMarker : public UAnimDatabaseFramesFunction
 {
 	GENERATED_BODY()
@@ -292,7 +292,7 @@ public:
 };
 
 /** Finds all the frames at the start of the given frame ranges */
-UCLASS(BlueprintType, Blueprintable, meta = (DisplayName = "Frame Range Starts"))
+UCLASS(BlueprintType, Blueprintable, meta = (DisplayName = "Frame Range Start Frames"))
 class UAnimDatabaseFramesFunction_FrameRangeStarts : public UAnimDatabaseFramesFunction
 {
 	GENERATED_BODY()
@@ -308,7 +308,7 @@ public:
 };
 
 /** Finds all the frames at the end of the given frame ranges */
-UCLASS(BlueprintType, Blueprintable, meta = (DisplayName = "Frame Range Ends"))
+UCLASS(BlueprintType, Blueprintable, meta = (DisplayName = "Frame Range End Frames"))
 class UAnimDatabaseFramesFunction_FrameRangeEnds : public UAnimDatabaseFramesFunction
 {
 	GENERATED_BODY()
@@ -789,7 +789,7 @@ public:
 };
 
 /** Returns the empty set of frame ranges */
-UCLASS(BlueprintType, Blueprintable, meta = (DisplayName = "Empty"))
+UCLASS(BlueprintType, Blueprintable, meta = (DisplayName = "Empty Frame Ranges"))
 class UAnimDatabaseFrameRangesFunction_Empty : public UAnimDatabaseFrameRangesFunction
 {
 	GENERATED_BODY()
@@ -800,7 +800,7 @@ public:
 };
 
 /** Returns the Frame Ranges given as input */
-UCLASS(BlueprintType, Blueprintable, meta = (DisplayName = "Identity"))
+UCLASS(BlueprintType, Blueprintable, meta = (DisplayName = "Identity Frame Ranges"))
 class UAnimDatabaseFrameRangesFunction_Identity : public UAnimDatabaseFrameRangesFunction
 {
 	GENERATED_BODY()
@@ -811,7 +811,7 @@ public:
 };
 
 /** Returns the Frame Ranges given by a set of frames */
-UCLASS(BlueprintType, Blueprintable, meta = (DisplayName = "Frames"))
+UCLASS(BlueprintType, Blueprintable, meta = (DisplayName = "Frames Frame Ranges"))
 class UAnimDatabaseFrameRangesFunction_Frames : public UAnimDatabaseFrameRangesFunction
 {
 	GENERATED_BODY()
@@ -825,7 +825,7 @@ public:
 };
 
 /** Ignores the frames given as input and returns all the frame ranges in the database */
-UCLASS(BlueprintType, Blueprintable, meta = (DisplayName = "All"))
+UCLASS(BlueprintType, Blueprintable, meta = (DisplayName = "All Frame Ranges"))
 class UAnimDatabaseFrameRangesFunction_All : public UAnimDatabaseFrameRangesFunction
 {
 	GENERATED_BODY()
@@ -836,7 +836,7 @@ public:
 };
 
 /** Finds all the frame ranges that are mirrored */
-UCLASS(BlueprintType, Blueprintable, meta = (DisplayName = "Mirrored"))
+UCLASS(BlueprintType, Blueprintable, meta = (DisplayName = "Mirrored Frame Ranges"))
 class UAnimDatabaseFrameRangesFunction_Mirrored : public UAnimDatabaseFrameRangesFunction
 {
 	GENERATED_BODY()
@@ -847,7 +847,7 @@ public:
 };
 
 /** Finds all the frame ranges that are not mirrored */
-UCLASS(BlueprintType, Blueprintable, meta = (DisplayName = "Not Mirrored"))
+UCLASS(BlueprintType, Blueprintable, meta = (DisplayName = "Not Mirrored Frame Ranges"))
 class UAnimDatabaseFrameRangesFunction_NotMirrored : public UAnimDatabaseFrameRangesFunction
 {
 	GENERATED_BODY()
@@ -858,7 +858,7 @@ public:
 };
 
 /** Finds all the frame ranges that have root motion enabled */
-UCLASS(BlueprintType, Blueprintable, meta = (DisplayName = "Root Motion Enabled"))
+UCLASS(BlueprintType, Blueprintable, meta = (DisplayName = "Root Motion Enabled Frame Ranges"))
 class UAnimDatabaseFrameRangesFunction_RootMotionEnabled : public UAnimDatabaseFrameRangesFunction
 {
 	GENERATED_BODY()
@@ -869,7 +869,7 @@ public:
 };
 
 /** Finds all the frame ranges that have root motion disabled */
-UCLASS(BlueprintType, Blueprintable, meta = (DisplayName = "Root Motion Disabled"))
+UCLASS(BlueprintType, Blueprintable, meta = (DisplayName = "Root Motion Disabled Frame Ranges"))
 class UAnimDatabaseFrameRangesFunction_RootMotionDisabled : public UAnimDatabaseFrameRangesFunction
 {
 	GENERATED_BODY()
@@ -880,7 +880,7 @@ public:
 };
 
 /** Finds all the frame ranges that have force root lock enabled */
-UCLASS(BlueprintType, Blueprintable, meta = (DisplayName = "Force Root Lock Enabled"))
+UCLASS(BlueprintType, Blueprintable, meta = (DisplayName = "Force Root Lock Enabled Frame Ranges"))
 class UAnimDatabaseFrameRangesFunction_ForceRootLockEnabled : public UAnimDatabaseFrameRangesFunction
 {
 	GENERATED_BODY()
@@ -891,7 +891,7 @@ public:
 };
 
 /** Finds all the frame ranges that have force root lock disabled */
-UCLASS(BlueprintType, Blueprintable, meta = (DisplayName = "Force Root Lock Disabled"))
+UCLASS(BlueprintType, Blueprintable, meta = (DisplayName = "Force Root Lock Disabled Frame Ranges"))
 class UAnimDatabaseFrameRangesFunction_ForceRootLockDisabled : public UAnimDatabaseFrameRangesFunction
 {
 	GENERATED_BODY()
@@ -902,7 +902,7 @@ public:
 };
 
 /** Finds all the frame ranges that are looped */
-UCLASS(BlueprintType, Blueprintable, meta = (DisplayName = "Looped"))
+UCLASS(BlueprintType, Blueprintable, meta = (DisplayName = "Looped Frame Ranges"))
 class UAnimDatabaseFrameRangesFunction_Looped : public UAnimDatabaseFrameRangesFunction
 {
 	GENERATED_BODY()
@@ -913,7 +913,7 @@ public:
 };
 
 /** Finds all the frame ranges that are not looped */
-UCLASS(BlueprintType, Blueprintable, meta = (DisplayName = "Not Looped"))
+UCLASS(BlueprintType, Blueprintable, meta = (DisplayName = "Not Looped Frame Ranges"))
 class UAnimDatabaseFrameRangesFunction_NotLooped : public UAnimDatabaseFrameRangesFunction
 {
 	GENERATED_BODY()
@@ -924,7 +924,7 @@ public:
 };
 
 /** Finds all the frame ranges corresponding to the given Anim Sequence */
-UCLASS(BlueprintType, Blueprintable, meta = (DisplayName = "Anim Sequence"))
+UCLASS(BlueprintType, Blueprintable, meta = (DisplayName = "Anim Sequence Frame Ranges"))
 class UAnimDatabaseFrameRangesFunction_AnimSequence : public UAnimDatabaseFrameRangesFunction
 {
 	GENERATED_BODY()
@@ -940,7 +940,7 @@ public:
 };
 
 /** Finds all the frame ranges corresponding to the given Anim Sequences */
-UCLASS(BlueprintType, Blueprintable, meta = (DisplayName = "Anim Sequences"))
+UCLASS(BlueprintType, Blueprintable, meta = (DisplayName = "Anim Sequences Frame Ranges"))
 class UAnimDatabaseFrameRangesFunction_AnimSequences : public UAnimDatabaseFrameRangesFunction
 {
 	GENERATED_BODY()
@@ -956,7 +956,7 @@ public:
 };
 
 /** Finds all the frame ranges with a given Anim Notify Class */
-UCLASS(BlueprintType, Blueprintable, meta = (DisplayName = "Anim Notify"))
+UCLASS(BlueprintType, Blueprintable, meta = (DisplayName = "Anim Notify Frame Ranges"))
 class UAnimDatabaseFrameRangesFunction_AnimNotify : public UAnimDatabaseFrameRangesFunction
 {
 	GENERATED_BODY()
@@ -972,7 +972,7 @@ public:
 };
 
 /** Finds all the frame ranges with a given Sync Marker */
-UCLASS(BlueprintType, Blueprintable, meta = (DisplayName = "Sync Marker"))
+UCLASS(BlueprintType, Blueprintable, meta = (DisplayName = "Sync Marker Frame Ranges"))
 class UAnimDatabaseFrameRangesFunction_SyncMarker : public UAnimDatabaseFrameRangesFunction
 {
 	GENERATED_BODY()
@@ -988,7 +988,7 @@ public:
 };
 
 /** Finds all the frame ranges with a given Anim Notify State Class */
-UCLASS(BlueprintType, Blueprintable, meta = (DisplayName = "Anim Notify State"))
+UCLASS(BlueprintType, Blueprintable, meta = (DisplayName = "Anim Notify State Frame Ranges"))
 class UAnimDatabaseFrameRangesFunction_AnimNotifyState : public UAnimDatabaseFrameRangesFunction
 {
 	GENERATED_BODY()
@@ -1004,7 +1004,7 @@ public:
 };
 
 /** Finds all the frame ranges which do not have the given Anim Notify State Class active */
-UCLASS(BlueprintType, Blueprintable, meta = (DisplayName = "Anim Notify State Not Active"))
+UCLASS(BlueprintType, Blueprintable, meta = (DisplayName = "Anim Notify State Not Active Frame Ranges"))
 class UAnimDatabaseFrameRangesFunction_AnimNotifyStateNotActive : public UAnimDatabaseFrameRangesFunction
 {
 	GENERATED_BODY()
@@ -1020,7 +1020,7 @@ public:
 };
 
 /** Finds all the frame ranges with any of the given Anim Notify State Classes */
-UCLASS(BlueprintType, Blueprintable, meta = (DisplayName = "Anim Notify State Union"))
+UCLASS(BlueprintType, Blueprintable, meta = (DisplayName = "Anim Notify State Union Frame Ranges"))
 class UAnimDatabaseFrameRangesFunction_AnimNotifyStateUnion : public UAnimDatabaseFrameRangesFunction
 {
 	GENERATED_BODY()
@@ -1036,7 +1036,7 @@ public:
 };
 
 /** Finds all the frame ranges with all of the given Anim Notify State Classes */
-UCLASS(BlueprintType, Blueprintable, meta = (DisplayName = "Anim Notify State Intersection"))
+UCLASS(BlueprintType, Blueprintable, meta = (DisplayName = "Anim Notify State Intersection Frame Ranges"))
 class UAnimDatabaseFrameRangesFunction_AnimNotifyStateIntersection : public UAnimDatabaseFrameRangesFunction
 {
 	GENERATED_BODY()
@@ -1052,7 +1052,7 @@ public:
 };
 
 /** Finds all the frame ranges with a given Anim Notify State classes for mirrored and un-mirrored sequences */
-UCLASS(BlueprintType, Blueprintable, meta = (DisplayName = "Anim Notify State with Mirrored"))
+UCLASS(BlueprintType, Blueprintable, meta = (DisplayName = "Anim Notify State with Mirrored Frame Ranges"))
 class UAnimDatabaseFrameRangesFunction_AnimNotifyStateWithMirrored : public UAnimDatabaseFrameRangesFunction
 {
 	GENERATED_BODY()
@@ -1071,7 +1071,7 @@ public:
 };
 
 /** Finds all the frame ranges with the given sub-string in the name */
-UCLASS(BlueprintType, Blueprintable, meta = (DisplayName = "Asset Name Contains"))
+UCLASS(BlueprintType, Blueprintable, meta = (DisplayName = "Asset Name Contains Frame Ranges"))
 class UAnimDatabaseFrameRangesFunction_AssetNameContains : public UAnimDatabaseFrameRangesFunction
 {
 	GENERATED_BODY()
@@ -1087,7 +1087,7 @@ public:
 };
 
 /** Finds all the frame ranges with any of the given sub-strings in the name */
-UCLASS(BlueprintType, Blueprintable, meta = (DisplayName = "Asset Name Contains One Of"))
+UCLASS(BlueprintType, Blueprintable, meta = (DisplayName = "Asset Name Contains One Of Frame Ranges"))
 class UAnimDatabaseFrameRangesFunction_AssetNameContainsOneOf : public UAnimDatabaseFrameRangesFunction
 {
 	GENERATED_BODY()
@@ -1103,7 +1103,7 @@ public:
 };
 
 /** Finds all the frame ranges with the given sub-string in the package */
-UCLASS(BlueprintType, Blueprintable, meta = (DisplayName = "Asset Package Contains"))
+UCLASS(BlueprintType, Blueprintable, meta = (DisplayName = "Asset Package Contains Frame Ranges"))
 class UAnimDatabaseFrameRangesFunction_AssetPackageContains : public UAnimDatabaseFrameRangesFunction
 {
 	GENERATED_BODY()
@@ -1119,7 +1119,7 @@ public:
 };
 
 /** Finds all the frame ranges with any of the given sub-strings in the package */
-UCLASS(BlueprintType, Blueprintable, meta = (DisplayName = "Asset Package Contains One Of"))
+UCLASS(BlueprintType, Blueprintable, meta = (DisplayName = "Asset Package Contains One Of Frame Ranges"))
 class UAnimDatabaseFrameRangesFunction_AssetPackageContainsOneOf : public UAnimDatabaseFrameRangesFunction
 {
 	GENERATED_BODY()
@@ -1135,7 +1135,7 @@ public:
 };
 
 /** Finds all the frame ranges that start with the given sub-string in the name */
-UCLASS(BlueprintType, Blueprintable, meta = (DisplayName = "Asset Name Starts With"))
+UCLASS(BlueprintType, Blueprintable, meta = (DisplayName = "Asset Name Starts With Frame Ranges"))
 class UAnimDatabaseFrameRangesFunction_AssetNameStartsWith : public UAnimDatabaseFrameRangesFunction
 {
 	GENERATED_BODY()
@@ -1151,7 +1151,7 @@ public:
 };
 
 /** Finds all the frame ranges that start with any of the given sub-strings in the name */
-UCLASS(BlueprintType, Blueprintable, meta = (DisplayName = "Asset Name Starts with One Of"))
+UCLASS(BlueprintType, Blueprintable, meta = (DisplayName = "Asset Name Starts With One Of Frame Ranges"))
 class UAnimDatabaseFrameRangesFunction_AssetNameStartsWithOneOf : public UAnimDatabaseFrameRangesFunction
 {
 	GENERATED_BODY()
@@ -1167,7 +1167,7 @@ public:
 };
 
 /** Applies a sequence of FrameRangesFunctions to filter ranges */
-UCLASS(BlueprintType, Blueprintable, meta = (DisplayName = "Filter"))
+UCLASS(BlueprintType, Blueprintable, meta = (DisplayName = "Filter Frame Ranges"))
 class UAnimDatabaseFrameRangesFunction_Filter : public UAnimDatabaseFrameRangesFunction
 {
 	GENERATED_BODY()
@@ -1183,7 +1183,7 @@ public:
 };
 
 /** Returns all the frame ranges in the database that are NOT given by the provided function  */
-UCLASS(BlueprintType, Blueprintable, meta = (DisplayName = "Not"))
+UCLASS(BlueprintType, Blueprintable, meta = (DisplayName = "Not Frame Ranges"))
 class UAnimDatabaseFrameRangesFunction_Not : public UAnimDatabaseFrameRangesFunction
 {
 	GENERATED_BODY()
@@ -1199,7 +1199,7 @@ public:
 };
 
 /** Finds all the frame ranges which are a union of the given sub-ranges */
-UCLASS(BlueprintType, Blueprintable, meta = (DisplayName = "Union"))
+UCLASS(BlueprintType, Blueprintable, meta = (DisplayName = "Union Frame Ranges"))
 class UAnimDatabaseFrameRangesFunction_Union : public UAnimDatabaseFrameRangesFunction
 {
 	GENERATED_BODY()
@@ -1215,7 +1215,7 @@ public:
 };
 
 /** Finds all the frame ranges which are an intersection of the given sub-ranges */
-UCLASS(BlueprintType, Blueprintable, meta = (DisplayName = "Intersection"))
+UCLASS(BlueprintType, Blueprintable, meta = (DisplayName = "Intersection Frame Ranges"))
 class UAnimDatabaseFrameRangesFunction_Intersection : public UAnimDatabaseFrameRangesFunction
 {
 	GENERATED_BODY()
@@ -1231,7 +1231,7 @@ public:
 };
 
 /** Finds all the frame ranges which are an difference of the given sub-ranges */
-UCLASS(BlueprintType, Blueprintable, meta = (DisplayName = "Difference"))
+UCLASS(BlueprintType, Blueprintable, meta = (DisplayName = "Difference Frame Ranges"))
 class UAnimDatabaseFrameRangesFunction_Difference : public UAnimDatabaseFrameRangesFunction
 {
 	GENERATED_BODY()
@@ -1247,7 +1247,7 @@ public:
 };
 
 /** Finds all the frame ranges which bounds another set of ranges */
-UCLASS(BlueprintType, Blueprintable, meta = (DisplayName = "Bounds"))
+UCLASS(BlueprintType, Blueprintable, meta = (DisplayName = "Bounds Frame Ranges"))
 class UAnimDatabaseFrameRangesFunction_Bounds : public UAnimDatabaseFrameRangesFunction
 {
 	GENERATED_BODY()
@@ -1263,7 +1263,7 @@ public:
 };
 
 /** Finds all the frame ranges which contain any part of the provided frame ranges */
-UCLASS(BlueprintType, Blueprintable, meta = (DisplayName = "Intersects"))
+UCLASS(BlueprintType, Blueprintable, meta = (DisplayName = "Intersects Frame Ranges"))
 class UAnimDatabaseFrameRangesFunction_Intersects : public UAnimDatabaseFrameRangesFunction
 {
 	GENERATED_BODY()
@@ -1279,7 +1279,7 @@ public:
 };
 
 /** Finds all the frame ranges which contain any part of the provided frames */
-UCLASS(BlueprintType, Blueprintable, meta = (DisplayName = "Intersects Frames"))
+UCLASS(BlueprintType, Blueprintable, meta = (DisplayName = "Intersects Frames Frame Ranges"))
 class UAnimDatabaseFrameRangesFunction_IntersectsFrames : public UAnimDatabaseFrameRangesFunction
 {
 	GENERATED_BODY()
@@ -1295,7 +1295,7 @@ public:
 };
 
 /** Finds all the frame ranges which bounds another set of frames */
-UCLASS(BlueprintType, Blueprintable, meta = (DisplayName = "Bounds Frames"))
+UCLASS(BlueprintType, Blueprintable, meta = (DisplayName = "Bounds Frames Frame Ranges"))
 class UAnimDatabaseFrameRangesFunction_BoundsFrames : public UAnimDatabaseFrameRangesFunction
 {
 	GENERATED_BODY()
@@ -1311,7 +1311,7 @@ public:
 };
 
 /** Frame ranges consisting of single frames spaced at a regular interval */
-UCLASS(BlueprintType, Blueprintable, meta = (DisplayName = "Interval"))
+UCLASS(BlueprintType, Blueprintable, meta = (DisplayName = "Interval Frame Ranges"))
 class UAnimDatabaseFrameRangesFunction_Interval : public UAnimDatabaseFrameRangesFunction
 {
 	GENERATED_BODY()
@@ -1327,7 +1327,7 @@ public:
 };
 
 /** Trims the given frame ranges */
-UCLASS(BlueprintType, Blueprintable, meta = (DisplayName = "Trim"))
+UCLASS(BlueprintType, Blueprintable, meta = (DisplayName = "Trim Frame Ranges"))
 class UAnimDatabaseFrameRangesFunction_Trim : public UAnimDatabaseFrameRangesFunction
 {
 	GENERATED_BODY()
@@ -1351,7 +1351,7 @@ public:
 };
 
 /** Pads the given frame ranges */
-UCLASS(BlueprintType, Blueprintable, meta = (DisplayName = "Pad"))
+UCLASS(BlueprintType, Blueprintable, meta = (DisplayName = "Pad Frame Ranges"))
 class UAnimDatabaseFrameRangesFunction_Pad : public UAnimDatabaseFrameRangesFunction
 {
 	GENERATED_BODY()
diff --git a/Engine/Plugins/Experimental/Animation/AnimGen/Source/AnimGen/Public/AnimGenAutoEncoder.h b/Engine/Plugins/Experimental/Animation/AnimGen/Source/AnimGen/Public/AnimGenAutoEncoder.h
--- a/Engine/Plugins/Experimental/Animation/AnimGen/Source/AnimGen/Public/AnimGenAutoEncoder.h
+++ b/Engine/Plugins/Experimental/Animation/AnimGen/Source/AnimGen/Public/AnimGenAutoEncoder.h
@@ -753,7 +753,7 @@ public:
 
 
 /** Debug Draw class for drawing contact timings */
-UCLASS(BlueprintType, Blueprintable, meta = (DisplayName = "Contact Timings"))
+UCLASS(BlueprintType, Blueprintable, meta = (DisplayName = "Contact Timings Debug Draw"))
 class UAnimGenAutoEncoderDebugDraw_ContactTimings : public UAnimGenAutoEncoderDebugDraw
 {
 	GENERATED_BODY()
diff --git a/Engine/Plugins/Experimental/Animation/AnimGen/Source/AnimGen/Public/AnimGenBehavior.h b/Engine/Plugins/Experimental/Animation/AnimGen/Source/AnimGen/Public/AnimGenBehavior.h
--- a/Engine/Plugins/Experimental/Animation/AnimGen/Source/AnimGen/Public/AnimGenBehavior.h
+++ b/Engine/Plugins/Experimental/Animation/AnimGen/Source/AnimGen/Public/AnimGenBehavior.h
@@ -57,7 +57,7 @@ public:
 };
 
 /** Behavior for uncontrolled generation. Produces a random-walk of the dataset. */
-UCLASS(BlueprintType, Blueprintable, meta = (DisplayName = "Uncontrolled"))
+UCLASS(BlueprintType, Blueprintable, meta = (DisplayName = "Uncontrolled Behavior"))
 class UAnimGenBehavior_Uncontrolled : public UAnimGenBehavior
 {
 	GENERATED_BODY()
@@ -74,14 +74,14 @@ public:
 };
 
 /** Behavior for idling. */
-UCLASS(BlueprintType, Blueprintable, meta = (DisplayName = "Idle"))
+UCLASS(BlueprintType, Blueprintable, meta = (DisplayName = "Idle Behavior"))
 class UAnimGenBehavior_Idle : public UAnimGenBehavior_Uncontrolled
 {
 	GENERATED_BODY()
 };
 
 /** Behavior for following a trajectory which can either be generated from Mover/CharacterMovementComponent or a path to follow */
-UCLASS(BlueprintType, Blueprintable, meta = (DisplayName = "Trajectory Follow"))
+UCLASS(BlueprintType, Blueprintable, meta = (DisplayName = "Trajectory Follow Behavior"))
 class UAnimGenBehavior_TrajectoryFollow : public UAnimGenBehavior
 {
 	GENERATED_BODY()
@@ -133,7 +133,7 @@ public:
 };
 
 /** Behavior for following a velocity and angular velocity. More useful for offline generation. */
-UCLASS(BlueprintType, Blueprintable, meta = (DisplayName = "Velocity Follow"))
+UCLASS(BlueprintType, Blueprintable, meta = (DisplayName = "Velocity Follow Behavior"))
 class UAnimGenBehavior_VelocityFollow : public UAnimGenBehavior
 {
 	GENERATED_BODY()
@@ -163,7 +163,7 @@ public:
 };
 
 /** Behavior for moving toward a target with an optional facing direction */
-UCLASS(BlueprintType, Blueprintable, meta = (DisplayName = "Move to Target"))
+UCLASS(BlueprintType, Blueprintable, meta = (DisplayName = "Move to Target Behavior"))
 class UAnimGenBehavior_MoveToTarget : public UAnimGenBehavior
 {
 	GENERATED_BODY()
@@ -196,7 +196,7 @@ public:
  * Behavior for following a trajectory while also interacting with an object or prop. This is useful for networked prop interactions where the 
  * trajectory may be coming from a montage .
  */
-UCLASS(BlueprintType, Blueprintable, DontCollapseCategories, meta = (DisplayName = "Trajectory Interaction"))
+UCLASS(BlueprintType, Blueprintable, DontCollapseCategories, meta = (DisplayName = "Trajectory Interaction Behavior"))
 class UAnimGenBehavior_TrajectoryInteraction : public UAnimGenBehavior
 {
 	GENERATED_BODY()
@@ -268,7 +268,7 @@ public:
 };
 
 /** Behavior that allows a user to supply tags to control the style of a sub-behavior. */
-UCLASS(BlueprintType, Blueprintable, meta = (DisplayName = "Tagged"))
+UCLASS(BlueprintType, Blueprintable, meta = (DisplayName = "Tagged Behavior"))
 class UAnimGenBehavior_Tagged : public UAnimGenBehavior
 {
 	GENERATED_BODY()
@@ -337,7 +337,7 @@ public:
 };
 
 /** Behavior made up of multiple other behaviors - user can also provide some tags to control the style. */
-UCLASS(BlueprintType, Blueprintable, meta = (DisplayName = "Multi"))
+UCLASS(BlueprintType, Blueprintable, meta = (DisplayName = "Multi Behavior"))
 class UAnimGenBehavior_Multi : public UAnimGenBehavior
 {
 	GENERATED_BODY()
