AnimDatabase: Added asset name ends with frame ranges function

#Synced-CL 52545394

diff --git a/Engine/Plugins/Experimental/Animation/AnimDatabase/Source/AnimDatabase/Private/AnimDatabaseFrameRanges.cpp b/Engine/Plugins/Experimental/Animation/AnimDatabase/Source/AnimDatabase/Private/AnimDatabaseFrameRanges.cpp
--- a/Engine/Plugins/Experimental/Animation/AnimDatabase/Source/AnimDatabase/Private/AnimDatabaseFrameRanges.cpp
+++ b/Engine/Plugins/Experimental/Animation/AnimDatabase/Source/AnimDatabase/Private/AnimDatabaseFrameRanges.cpp
@@ -3192,6 +3192,16 @@ FAnimDatabaseFrameRanges UAnimDatabaseFrameRangesFunction_AssetNameStartsWithOne
 	return UAnimDatabaseFrameRangesLibrary::MakeFrameRangesFromAssetNameStartsWithOneOfArrayView(Database, FrameRanges, AssetNamePrefixes);
 }
 
+FAnimDatabaseFrameRanges UAnimDatabaseFrameRangesFunction_AssetNameEndsWith::MakeFrameRanges_Implementation(const UAnimDatabase* Database, const FAnimDatabaseFrameRanges& FrameRanges) const
+{
+	return UAnimDatabaseFrameRangesLibrary::MakeFrameRangesFromAssetNameEndsWithStringView(Database, FrameRanges, AssetNameSuffix);
+}
+
+FAnimDatabaseFrameRanges UAnimDatabaseFrameRangesFunction_AssetNameEndsWithOneOf::MakeFrameRanges_Implementation(const UAnimDatabase* Database, const FAnimDatabaseFrameRanges& FrameRanges) const
+{
+	return UAnimDatabaseFrameRangesLibrary::MakeFrameRangesFromAssetNameEndsWithOneOfArrayView(Database, FrameRanges, AssetNameSuffixes);
+}
+
 FAnimDatabaseFrameRanges UAnimDatabaseFrameRangesFunction_Filter::MakeFrameRanges_Implementation(const UAnimDatabase* Database, const FAnimDatabaseFrameRanges& FrameRanges) const
 {
 	FAnimDatabaseFrameRanges OutFrameRanges = FrameRanges;
diff --git a/Engine/Plugins/Experimental/Animation/AnimDatabase/Source/AnimDatabase/Public/AnimDatabaseFrameRanges.h b/Engine/Plugins/Experimental/Animation/AnimDatabase/Source/AnimDatabase/Public/AnimDatabaseFrameRanges.h
--- a/Engine/Plugins/Experimental/Animation/AnimDatabase/Source/AnimDatabase/Public/AnimDatabaseFrameRanges.h
+++ b/Engine/Plugins/Experimental/Animation/AnimDatabase/Source/AnimDatabase/Public/AnimDatabaseFrameRanges.h
@@ -1166,6 +1166,38 @@ public:
 	UE_API virtual FAnimDatabaseFrameRanges MakeFrameRanges_Implementation(const UAnimDatabase* Database, const FAnimDatabaseFrameRanges& FrameRanges) const override;
 };
 
+/** Finds all the frame ranges that end with the given sub-string in the name */
+UCLASS(BlueprintType, Blueprintable, meta = (DisplayName = "Asset Name Ends With Frame Ranges"))
+class UAnimDatabaseFrameRangesFunction_AssetNameEndsWith : public UAnimDatabaseFrameRangesFunction
+{
+	GENERATED_BODY()
+
+public:
+
+	UPROPERTY(EditAnywhere, Category = "Settings")
+	FString AssetNameSuffix;
+
+public:
+
+	UE_API virtual FAnimDatabaseFrameRanges MakeFrameRanges_Implementation(const UAnimDatabase* Database, const FAnimDatabaseFrameRanges& FrameRanges) const override;
+};
+
+/** Finds all the frame ranges that end with any of the given sub-strings in the name */
+UCLASS(BlueprintType, Blueprintable, meta = (DisplayName = "Asset Name Ends With One Of Frame Ranges"))
+class UAnimDatabaseFrameRangesFunction_AssetNameEndsWithOneOf : public UAnimDatabaseFrameRangesFunction
+{
+	GENERATED_BODY()
+
+public:
+
+	UPROPERTY(EditAnywhere, Category = "Settings")
+	TArray<FString> AssetNameSuffixes;
+
+public:
+
+	UE_API virtual FAnimDatabaseFrameRanges MakeFrameRanges_Implementation(const UAnimDatabase* Database, const FAnimDatabaseFrameRanges& FrameRanges) const override;
+};
+
 /** Applies a sequence of FrameRangesFunctions to filter ranges */
 UCLASS(BlueprintType, Blueprintable, meta = (DisplayName = "Filter Frame Ranges"))
 class UAnimDatabaseFrameRangesFunction_Filter : public UAnimDatabaseFrameRangesFunction
