AnimGen: Added attribute smoothing to help reduce the noise

#Synced-CL 52573539

diff --git a/Engine/Plugins/Experimental/Animation/AnimDatabase/Source/AnimDatabase/Private/AnimDatabasePose.cpp b/Engine/Plugins/Experimental/Animation/AnimDatabase/Source/AnimDatabase/Private/AnimDatabasePose.cpp
--- a/Engine/Plugins/Experimental/Animation/AnimDatabase/Source/AnimDatabase/Private/AnimDatabasePose.cpp
+++ b/Engine/Plugins/Experimental/Animation/AnimDatabase/Source/AnimDatabase/Private/AnimDatabasePose.cpp
@@ -3267,6 +3267,8 @@ namespace UE::AnimDatabase
 			const int32 FrameNum = InOutPose.GetFrameNum();
 			const int32 BoneNum = InOutPose.GetBoneNum();
 
+			const float Factor = 1.0f - FMath::InvExpApprox(DeltaTime / FMath::Max(PoseSmoothingTime, UE_SMALL_NUMBER));
+
 			for (int32 FrameIdx = 0; FrameIdx < FrameNum; FrameIdx++)
 			{
 				for (int32 BoneIdx = 0; BoneIdx < BoneNum; BoneIdx++)
@@ -3280,8 +3282,6 @@ namespace UE::AnimDatabase
 					const FQuat4f BoneRotationIntegrated = FQuat4f::MakeFromRotationVector(DeltaTime * AngularVelocity) * InOutPose.BoneRotations[FrameIdx][BoneIdx];
 					const FVector3f BoneScaleIntegrated = Math::VectorExpSafe(DeltaTime * ScalarVelocity) * InOutPose.BoneScales[FrameIdx][BoneIdx];
 
-					const float Factor = 1.0f - FMath::InvExpApprox(DeltaTime / FMath::Max(PoseSmoothingTime, UE_SMALL_NUMBER));
-
 					InOutPose.BoneLocations[FrameIdx][BoneIdx] = FMath::Lerp(BoneLocationIntegrated, InTargetPose.BoneLocations[FrameIdx][BoneIdx], Factor);
 					InOutPose.BoneRotations[FrameIdx][BoneIdx] = FQuat4f::Slerp(BoneRotationIntegrated, InTargetPose.BoneRotations[FrameIdx][BoneIdx], Factor);
 					InOutPose.BoneScales[FrameIdx][BoneIdx] = Math::VectorEerp(BoneScaleIntegrated, InTargetPose.BoneScales[FrameIdx][BoneIdx], Factor);
@@ -3391,6 +3391,101 @@ namespace UE::AnimDatabase
 			}
 		}
 
+		void DampAttributesInPlace(
+			const FPoseAttributeDataView& InOutAttributes,
+			const FPoseAttributeDataConstView& InTargetAttributes,
+			const float DeltaTime,
+			const float SmoothingTime)
+		{
+			TRACE_CPUPROFILER_EVENT_SCOPE(AnimDatabase::PoseData::DampAttributesInPlace);
+
+			check(InOutAttributes.GetFrameNum() == InTargetAttributes.GetFrameNum());
+			check(Learning::Array::Equal(InOutAttributes.GetAttributeTypes(), InTargetAttributes.GetAttributeTypes()));
+
+			const int32 FrameNum = InOutAttributes.GetFrameNum();
+			const int32 AttributeNum = InOutAttributes.GetAttributeNum();
+
+			const float Factor = 1.0f - FMath::InvExpApprox(DeltaTime / FMath::Max(SmoothingTime, UE_SMALL_NUMBER));
+
+			for (int32 FrameIdx = 0; FrameIdx < FrameNum; FrameIdx++)
+			{
+				for (int32 AttributeIdx = 0; AttributeIdx < AttributeNum; AttributeIdx++)
+				{
+					const EAnimDatabaseAttributeType AttributeType = InOutAttributes.GetAttributeType(AttributeIdx);
+					const bool bCurrentActive = InOutAttributes.GetAttributeActive(FrameIdx, AttributeIdx);
+					const bool bTargetActive = InTargetAttributes.GetAttributeActive(FrameIdx, AttributeIdx);
+
+					InOutAttributes.SetAttributeActive(FrameIdx, AttributeIdx, bTargetActive);
+
+					if (bTargetActive && bCurrentActive)
+					{
+						switch (AttributeType)
+						{
+						case EAnimDatabaseAttributeType::Null: break;
+						case EAnimDatabaseAttributeType::Bool: InOutAttributes.SetBool(FrameIdx, AttributeIdx, InTargetAttributes.GetBool(FrameIdx, AttributeIdx)); break;
+						case EAnimDatabaseAttributeType::Float: InOutAttributes.SetFloat(FrameIdx, AttributeIdx, FMath::Lerp(InOutAttributes.GetFloat(FrameIdx, AttributeIdx), InTargetAttributes.GetFloat(FrameIdx, AttributeIdx), Factor)); break;
+						case EAnimDatabaseAttributeType::Angle: InOutAttributes.SetAngle(FrameIdx, AttributeIdx, InOutAttributes.GetAngle(FrameIdx, AttributeIdx) + Factor * FMath::FindDeltaAngleRadians(InOutAttributes.GetAngle(FrameIdx, AttributeIdx), InTargetAttributes.GetAngle(FrameIdx, AttributeIdx))); break;
+						case EAnimDatabaseAttributeType::Location: InOutAttributes.SetLocation(FrameIdx, AttributeIdx, FMath::Lerp(InOutAttributes.GetLocation(FrameIdx, AttributeIdx), InTargetAttributes.GetLocation(FrameIdx, AttributeIdx), Factor)); break;
+						case EAnimDatabaseAttributeType::Rotation: InOutAttributes.SetRotation(FrameIdx, AttributeIdx, FQuat4f::Slerp(InOutAttributes.GetRotation(FrameIdx, AttributeIdx), InTargetAttributes.GetRotation(FrameIdx, AttributeIdx), Factor)); break;
+						case EAnimDatabaseAttributeType::Scale: InOutAttributes.SetScale(FrameIdx, AttributeIdx, Math::VectorEerp(InOutAttributes.GetScale(FrameIdx, AttributeIdx), InTargetAttributes.GetScale(FrameIdx, AttributeIdx), Factor)); break;
+						case EAnimDatabaseAttributeType::LinearVelocity: InOutAttributes.SetLinearVelocity(FrameIdx, AttributeIdx, FMath::Lerp(InOutAttributes.GetLinearVelocity(FrameIdx, AttributeIdx), InTargetAttributes.GetLinearVelocity(FrameIdx, AttributeIdx), Factor)); break;
+						case EAnimDatabaseAttributeType::AngularVelocity: InOutAttributes.SetAngularVelocity(FrameIdx, AttributeIdx, FMath::Lerp(InOutAttributes.GetAngularVelocity(FrameIdx, AttributeIdx), InTargetAttributes.GetAngularVelocity(FrameIdx, AttributeIdx), Factor)); break;
+						case EAnimDatabaseAttributeType::ScalarVelocity: InOutAttributes.SetScalarVelocity(FrameIdx, AttributeIdx, FMath::Lerp(InOutAttributes.GetScalarVelocity(FrameIdx, AttributeIdx), InTargetAttributes.GetScalarVelocity(FrameIdx, AttributeIdx), Factor)); break;
+						case EAnimDatabaseAttributeType::Direction: InOutAttributes.SetDirection(FrameIdx, AttributeIdx, FMath::Lerp(InOutAttributes.GetDirection(FrameIdx, AttributeIdx), InTargetAttributes.GetDirection(FrameIdx, AttributeIdx), Factor).GetSafeNormal(UE_SMALL_NUMBER, FVector3f::ForwardVector)); break;
+						case EAnimDatabaseAttributeType::Transform:
+						{
+							const FTransform3f TargetTransform = InTargetAttributes.GetTransform(FrameIdx, AttributeIdx);
+							const FTransform3f CurrentTransform = InOutAttributes.GetTransform(FrameIdx, AttributeIdx);
+							FTransform3f OutTransform = FTransform3f::Identity;
+							OutTransform.SetLocation(FMath::Lerp(CurrentTransform.GetLocation(), TargetTransform.GetLocation(), Factor));
+							OutTransform.SetRotation(FQuat4f::Slerp(CurrentTransform.GetRotation(), TargetTransform.GetRotation(), Factor));
+							OutTransform.SetScale3D(Math::VectorEerp(CurrentTransform.GetScale3D(), TargetTransform.GetScale3D(), Factor));
+							InOutAttributes.SetTransform(FrameIdx, AttributeIdx, OutTransform);
+							break;
+						}
+
+						case EAnimDatabaseAttributeType::Event:
+						{
+							bool bTargetEventKnown = false, bCurrentEventKnown = false;
+							float TargetTimeUntilEvent = UE_MAX_FLT, CurrentTimeUntilEvent = UE_MAX_FLT;
+							InTargetAttributes.GetEvent(bTargetEventKnown, TargetTimeUntilEvent, FrameIdx, AttributeIdx);
+							InOutAttributes.GetEvent(bCurrentEventKnown, CurrentTimeUntilEvent, FrameIdx, AttributeIdx);
+
+							if (bTargetEventKnown)
+							{
+								if (bCurrentEventKnown)
+								{
+									InOutAttributes.SetEvent(FrameIdx, AttributeIdx, true, FMath::Lerp(CurrentTimeUntilEvent, TargetTimeUntilEvent, Factor));
+								}
+								else
+								{
+									InOutAttributes.SetEvent(FrameIdx, AttributeIdx, true, TargetTimeUntilEvent);
+								}
+							}
+							else
+							{
+								InOutAttributes.SetEvent(FrameIdx, AttributeIdx, false, UE_MAX_FLT);
+							}
+							break;
+						}
+						default:
+							checkNoEntry();
+						}
+					}
+					else
+					{
+						UE::Learning::Array::Copy(
+							InOutAttributes.AttributeData[FrameIdx].Slice(
+								InOutAttributes.GetAttributeOffset(AttributeIdx),
+								InOutAttributes.GetAttributeSize(AttributeIdx)),
+							InTargetAttributes.AttributeData[FrameIdx].Slice(
+								InTargetAttributes.GetAttributeOffset(AttributeIdx),
+								InTargetAttributes.GetAttributeSize(AttributeIdx)));
+					}
+				}
+			}
+		}
+
 		namespace ContactPinning
 		{
 			FContactStateConstView FContactStateConstView::Slice(const int32 FrameStart, const int32 FrameNum) const
diff --git a/Engine/Plugins/Experimental/Animation/AnimDatabase/Source/AnimDatabase/Public/AnimDatabasePose.h b/Engine/Plugins/Experimental/Animation/AnimDatabase/Source/AnimDatabase/Public/AnimDatabasePose.h
--- a/Engine/Plugins/Experimental/Animation/AnimDatabase/Source/AnimDatabase/Public/AnimDatabasePose.h
+++ b/Engine/Plugins/Experimental/Animation/AnimDatabase/Source/AnimDatabase/Public/AnimDatabasePose.h
@@ -780,6 +780,13 @@ namespace UE::AnimDatabase
 			const float AngularVelocitySmoothingThreshold = UE_PI / 8.0f,
 			const float ScalarVelocitySmoothingThreshold = 0.5f);
 
+		/** Update attributes in-place damping them toward the target attributes with the given smoothing time */
+		UE_API void DampAttributesInPlace(
+			const FPoseAttributeDataView& InOutAttributes,
+			const FPoseAttributeDataConstView& InTargetAttributes,
+			const float DeltaTime,
+			const float SmoothingTime);
+
 		/** Helper functions for performing contact pinning */
 		namespace ContactPinning
 		{
diff --git a/Engine/Plugins/Experimental/Animation/AnimGen/Source/AnimGen/Private/AnimNode_AnimGenController.cpp b/Engine/Plugins/Experimental/Animation/AnimGen/Source/AnimGen/Private/AnimNode_AnimGenController.cpp
--- a/Engine/Plugins/Experimental/Animation/AnimGen/Source/AnimGen/Private/AnimNode_AnimGenController.cpp
+++ b/Engine/Plugins/Experimental/Animation/AnimGen/Source/AnimGen/Private/AnimNode_AnimGenController.cpp
@@ -140,13 +140,18 @@ void FAnimNode_AnimGenController::Update_AnyThread(const FAnimationUpdateContext
 
 	const FAnimGenControlSchema& ControlSchema = ActiveController->ControlSchema;
 	const FAnimGenControlSchemaElement& ControlSchemaElement = ActiveController->ControlSchemaElement;
-	check(ControlSchema.ObservationSchema.IsValid() && ControlSchema.ObservationSchema->Schema.IsValid(ControlSchemaElement.SchemaElement));
+	check(ControlSchema.IsValid() && ControlSchema.IsElementValid(ControlSchemaElement));
 
-	if (!UAnimGenControls::ValidateControlObjectMatchesSchema(ControlSchema, ControlSchemaElement, ControlObject, ControlObjectElement))
+	if (!UAnimGenControls::ValidateControlObjectMatchesSchema(
+		ControlSchema,
+		ControlSchemaElement,
+		ControlObject,
+		ControlObjectElement))
 	{
 		if (ControlObject.IsValid())
 		{
 			FWriteScopeLock ObjectLock(ControlObject.ObservationObject->Lock);
+
 			ControlObject.ObservationObject->Object.Reset();
 		}
 
@@ -155,28 +160,39 @@ void FAnimNode_AnimGenController::Update_AnyThread(const FAnimationUpdateContext
 		return;
 	}
 
-	bControlsValid = true;
-
 	// Set Control Vector
 
 	{
 		FReadScopeLock ObjectLock(ControlObject.ObservationObject->Lock);
 
-		UE::Learning::Observation::SetVectorFromObject(
-			ControlVector[0],
-			ControlSchema.ObservationSchema->Schema,
-			ControlSchemaElement.SchemaElement,
-			ControlObject.ObservationObject->Object,
-			ControlObjectElement.ObjectElement);
+		if (ControlObject.IsValid() && ControlObject.IsElementValid(ControlObjectElement))
+		{
+			UE::Learning::Observation::SetVectorFromObject(
+				ControlVector[0],
+				ControlSchema.ObservationSchema->Schema,
+				ControlSchemaElement.SchemaElement,
+				ControlObject.ObservationObject->Object,
+				ControlObjectElement.ObjectElement);
+
+			bControlsValid = true;
+		}
+		else
+		{
+			bPoseStateRequiresReset = true;
+			bControlsValid = false;
+			return;
+		}
 	}
 
 	// Reset Control Object
 
 	{
 		FWriteScopeLock ObjectLock(ControlObject.ObservationObject->Lock);
+
 		ControlObject.ObservationObject->Object.Reset();
 	}
 
+
 	// Clamp Control Vectors
 
 	UE::Learning::Observation::ClampToObservationDistributionMinMax(
@@ -502,11 +518,11 @@ void FAnimNode_AnimGenController::Update_AnyThread(const FAnimationUpdateContext
 		{ (FQuat4f)PoseRootTransform.GetRotation() });
 
 	PoseState.CurrPoseData->RootData = OutputPoseData.RootData;
-	PoseState.CurrPoseData->AttributeData = OutputPoseData.AttributeData;
 
 	if (bPoseStateRequiresReset)
 	{
 		PoseState.CurrPoseData->LocalBoneData = OutputPoseData.LocalBoneData;
+		PoseState.CurrPoseData->AttributeData = OutputPoseData.AttributeData;
 	}
 	else
 	{
@@ -516,6 +532,12 @@ void FAnimNode_AnimGenController::Update_AnyThread(const FAnimationUpdateContext
 			DeltaTime,
 			MinPoseSmoothingTime,
 			MaxPoseSmoothingTime);
+
+		UE::AnimDatabase::PoseData::DampAttributesInPlace(
+			PoseState.CurrPoseData->AttributeData.View(),
+			OutputPoseData.AttributeData.ConstView(),
+			DeltaTime,
+			AttributeSmoothingTime);
 	}
 
 	// Update Delta Time
diff --git a/Engine/Plugins/Experimental/Animation/AnimGen/Source/AnimGen/Public/AnimNode_AnimGenController.h b/Engine/Plugins/Experimental/Animation/AnimGen/Source/AnimGen/Public/AnimNode_AnimGenController.h
--- a/Engine/Plugins/Experimental/Animation/AnimGen/Source/AnimGen/Public/AnimNode_AnimGenController.h
+++ b/Engine/Plugins/Experimental/Animation/AnimGen/Source/AnimGen/Public/AnimNode_AnimGenController.h
@@ -85,6 +85,10 @@ public:
 	UPROPERTY(EditAnywhere, Category = Curves)
 	TMap<FName, FName> BoneTransformMap;
 
+	/** Smoothing to apply to predicted attributes */
+	UPROPERTY(EditAnywhere, Category = Settings, meta = (ClampMin = "0.0", UIMin = "0.0"))
+	float AttributeSmoothingTime = 0.0f;
+
 private:
 
 	/** Internal buffer of AnimNotifyEvents */
