AnimDatabase: Adjusted how bone weights are computed to improve autoencoder accuracy on foot joints

#Synced-CL 52573539

diff --git a/Engine/Plugins/Experimental/Animation/AnimDatabase/Source/AnimDatabase/Private/AnimDatabasePose.cpp b/Engine/Plugins/Experimental/Animation/AnimDatabase/Source/AnimDatabase/Private/AnimDatabasePose.cpp
--- a/Engine/Plugins/Experimental/Animation/AnimDatabase/Source/AnimDatabase/Private/AnimDatabasePose.cpp
+++ b/Engine/Plugins/Experimental/Animation/AnimDatabase/Source/AnimDatabase/Private/AnimDatabasePose.cpp
@@ -2118,10 +2118,12 @@ namespace UE::AnimDatabase
 			TLearningArrayView<1, float> OutPoseVectorWeights,
 			const TLearningArrayView<1, const FTransform> ReferenceTransforms,
 			const TLearningArrayView<1, const int32> BoneParents,
+			const Learning::FIndexSet RequiredIndices,
 			const Learning::FIndexSet BoneLocationIndices,
 			const Learning::FIndexSet BoneRotationIndices,
 			const Learning::FIndexSet BoneScaleIndices,
-			const float RootWeight)
+			const float RootWeight,
+			const float BoneBaseWeightMultiplier)
 		{
 			Learning::Array::Set(OutPoseVectorWeights, 1.0f);
 			Learning::Array::Set(OutPoseVectorWeights.Slice(0, 12), RootWeight);
@@ -2140,7 +2142,7 @@ namespace UE::AnimDatabase
 			Learning::Array::Copy(GlobalReferenceTransforms, ReferenceTransforms);
 			for (int32 BoneIdx = 0; BoneIdx < BoneNum; BoneIdx++)
 			{
-				if (BoneParents[BoneIdx] != -1)
+				if (BoneParents[BoneIdx] != INDEX_NONE)
 				{
 					GlobalReferenceTransforms[BoneIdx] = GlobalReferenceTransforms[BoneIdx] * GlobalReferenceTransforms[BoneParents[BoneIdx]];
 				}
@@ -2167,12 +2169,10 @@ namespace UE::AnimDatabase
 						// Assume that the diameter of the cylinder is approximately 1/3 of the bone length
 						const float CylinderDiameterFromBoneLength = 0.33f;
 
-						// Approximate area with cylinder (minus caps) where diameter is half of the height
 						const float BoneLength = ReferenceTransforms[ChildIdx].GetLocation().Length();
 						const float CylinderHeight = BoneLength;
 						const float CylinderDiameter = BoneLength * CylinderDiameterFromBoneLength;
 						BoneSurfaceAreas[BoneIdx] += UE_PI * CylinderHeight * CylinderDiameter;
-
 						BoneAverageChildDistances[BoneIdx] += BoneLength / 2.0f;
 						ActiveChildNum++;
 					}
@@ -2190,24 +2190,26 @@ namespace UE::AnimDatabase
 			BoneDescendants.SetNumUninitialized({ BoneNum, BoneNum });
 			Math::BoneDescedantsMatrix(BoneDescendants, BoneParents);
 
-			// Assume all joints have a cylinder attached with 4cm diameter and 6cm height, which is 10cm away
-			const float BaseCylinderSurfaceArea = UE_PI * 6.0f * 4.0f;
+			// Assume all joints have a cylinder attached with 10cm diameter and 20cm height, which is 10cm away
+			const float BaseCylinderHeight = 20.0f;
+			const float BaseCylinderDiameter = 10.0f;
+			const float BaseCylinderSurfaceArea = UE_PI * BaseCylinderHeight * BaseCylinderDiameter;
 			const float BaseCylinderDistance = 10.0f;
 
 			// Find Location Weights
 
 			TLearningArray<1, float> BoneLocationWeights;
 			BoneLocationWeights.SetNumUninitialized({ BoneLocationNum });
-			float BoneLocationWeightTotal = UE_KINDA_SMALL_NUMBER;
+			float BoneLocationWeightTotal = 0.0f;
 
 			for (int32 LocBoneIdx = 0; LocBoneIdx < BoneLocationNum; LocBoneIdx++)
 			{
 				const int32 BoneIdx = BoneLocationIndices[LocBoneIdx];
-				BoneLocationWeights[LocBoneIdx] = BaseCylinderSurfaceArea;
+				BoneLocationWeights[LocBoneIdx] = BoneBaseWeightMultiplier * BaseCylinderSurfaceArea;
 
 				for (int32 DescIdx = 0; DescIdx < BoneNum; DescIdx++)
 				{
-					if (BoneIdx == DescIdx || BoneDescendants[BoneIdx][DescIdx])
+					if (BoneIdx == DescIdx || (BoneDescendants[BoneIdx][DescIdx] && RequiredIndices.Contains(DescIdx)))
 					{
 						BoneLocationWeights[LocBoneIdx] += BoneSurfaceAreas[DescIdx];
 					}
@@ -2218,28 +2220,27 @@ namespace UE::AnimDatabase
 
 			for (int32 LocBoneIdx = 0; LocBoneIdx < BoneLocationNum; LocBoneIdx++)
 			{
-				BoneLocationWeights[LocBoneIdx] = BoneLocationNum * (BoneLocationWeights[LocBoneIdx] / BoneLocationWeightTotal);
+				BoneLocationWeights[LocBoneIdx] = BoneLocationNum * (BoneLocationWeights[LocBoneIdx] / FMath::Max(BoneLocationWeightTotal, UE_SMALL_NUMBER));
 			}
 
 			// Find Rotation Weights
 
 			TLearningArray<1, float> BoneRotationWeights;
 			BoneRotationWeights.SetNumUninitialized({ BoneRotationNum });
-			float BoneRotationWeightTotal = UE_KINDA_SMALL_NUMBER;
+			float BoneRotationWeightTotal = 0.0f;
 
 			for (int32 RotBoneIdx = 0; RotBoneIdx < BoneRotationNum; RotBoneIdx++)
 			{
 				const int32 BoneIdx = BoneRotationIndices[RotBoneIdx];
-				BoneRotationWeights[RotBoneIdx] = BaseCylinderDistance * BaseCylinderSurfaceArea;
+				BoneRotationWeights[RotBoneIdx] = BoneBaseWeightMultiplier * BaseCylinderDistance * BaseCylinderSurfaceArea;
 
 				for (int32 DescIdx = 0; DescIdx < BoneNum; DescIdx++)
 				{
-					if (BoneIdx == DescIdx || BoneDescendants[BoneIdx][DescIdx])
+					if (BoneIdx == DescIdx || (BoneDescendants[BoneIdx][DescIdx] && RequiredIndices.Contains(DescIdx)))
 					{
-						// We assume a minimum distance of 5cm to make sure joints on top of each other still have an impact
 						const float Distance = FVector::Distance(
 							GlobalReferenceTransforms[BoneIdx].GetLocation(),
-							GlobalReferenceTransforms[DescIdx].GetLocation()) + BoneAverageChildDistances[BoneIdx];
+							GlobalReferenceTransforms[DescIdx].GetLocation()) + BoneAverageChildDistances[DescIdx];
 
 						BoneRotationWeights[RotBoneIdx] += Distance * BoneSurfaceAreas[DescIdx];
 					}
@@ -2250,7 +2251,7 @@ namespace UE::AnimDatabase
 
 			for (int32 RotBoneIdx = 0; RotBoneIdx < BoneRotationNum; RotBoneIdx++)
 			{
-				BoneRotationWeights[RotBoneIdx] = BoneRotationNum * (BoneRotationWeights[RotBoneIdx] / BoneRotationWeightTotal);
+				BoneRotationWeights[RotBoneIdx] = BoneRotationNum * (BoneRotationWeights[RotBoneIdx] / FMath::Max(BoneRotationWeightTotal, UE_SMALL_NUMBER));
 			}
 
 			// Insert into pose vector
diff --git a/Engine/Plugins/Experimental/Animation/AnimDatabase/Source/AnimDatabase/Public/AnimDatabasePose.h b/Engine/Plugins/Experimental/Animation/AnimDatabase/Source/AnimDatabase/Public/AnimDatabasePose.h
--- a/Engine/Plugins/Experimental/Animation/AnimDatabase/Source/AnimDatabase/Public/AnimDatabasePose.h
+++ b/Engine/Plugins/Experimental/Animation/AnimDatabase/Source/AnimDatabase/Public/AnimDatabasePose.h
@@ -623,10 +623,12 @@ namespace UE::AnimDatabase
 			TLearningArrayView<1, float> OutPoseVectorWeights,
 			const TLearningArrayView<1, const FTransform> ReferenceTransforms,
 			const TLearningArrayView<1, const int32> BoneParents,
+			const Learning::FIndexSet RequiredIndices,
 			const Learning::FIndexSet BoneLocationIndices,
 			const Learning::FIndexSet BoneRotationIndices,
 			const Learning::FIndexSet BoneScaleIndices,
-			const float RootWeight = 10.0f);
+			const float RootWeight = 10.0f,
+			const float BoneBaseWeightMultiplier = 1.0f);
 
 		/** Makes the provided pose data loop seamlessly */
 		UE_API void MakeLooped(
diff --git a/Engine/Plugins/Experimental/Animation/AnimGen/Source/AnimGen/Public/AnimGenAutoEncoder.h b/Engine/Plugins/Experimental/Animation/AnimGen/Source/AnimGen/Public/AnimGenAutoEncoder.h
--- a/Engine/Plugins/Experimental/Animation/AnimGen/Source/AnimGen/Public/AnimGenAutoEncoder.h
+++ b/Engine/Plugins/Experimental/Animation/AnimGen/Source/AnimGen/Public/AnimGenAutoEncoder.h
@@ -375,6 +375,14 @@ public:
 	UPROPERTY(EditAnywhere, Category = "Data")
 	bool bExcludeVirtualBones = true;
 
+	/** Increase this if the root motion is not being accurately reproduced */
+	UPROPERTY(EditAnywhere, Category = "Data", meta = (ClampMin = "0.0", UIMin = "0.0"))
+	float RootWeightMultiplier = 10.0f;
+
+	/** Increase this if the network is not accurately reproducing bones at the end of joint chains */
+	UPROPERTY(EditAnywhere, Category = "Data", meta = (ClampMin = "0.0", UIMin = "0.0"))
+	float BaseWeightMultiplier = 1.0f;
+
 public:
 
 	/** List of Frame Attributes to train the network to reproduce */
diff --git a/Engine/Plugins/Experimental/Animation/AnimGen/Source/AnimGenEditor/Private/AnimGenEditorAutoEncoderToolkit.cpp b/Engine/Plugins/Experimental/Animation/AnimGen/Source/AnimGenEditor/Private/AnimGenEditorAutoEncoderToolkit.cpp
--- a/Engine/Plugins/Experimental/Animation/AnimGen/Source/AnimGenEditor/Private/AnimGenEditorAutoEncoderToolkit.cpp
+++ b/Engine/Plugins/Experimental/Animation/AnimGen/Source/AnimGenEditor/Private/AnimGenEditorAutoEncoderToolkit.cpp
@@ -191,6 +191,8 @@ namespace UE::AnimGen::Editor
 		DataCategory.AddProperty(GET_MEMBER_NAME_CHECKED(UAnimGenAutoEncoderTrainingSettings, FrameRanges), UAnimGenAutoEncoderTrainingSettings::StaticClass());
 		DataCategory.AddProperty(GET_MEMBER_NAME_CHECKED(UAnimGenAutoEncoderTrainingSettings, ExcludedBones), UAnimGenAutoEncoderTrainingSettings::StaticClass());
 		DataCategory.AddProperty(GET_MEMBER_NAME_CHECKED(UAnimGenAutoEncoderTrainingSettings, bExcludeVirtualBones), UAnimGenAutoEncoderTrainingSettings::StaticClass());
+		DataCategory.AddProperty(GET_MEMBER_NAME_CHECKED(UAnimGenAutoEncoderTrainingSettings, RootWeightMultiplier), UAnimGenAutoEncoderTrainingSettings::StaticClass());
+		DataCategory.AddProperty(GET_MEMBER_NAME_CHECKED(UAnimGenAutoEncoderTrainingSettings, BaseWeightMultiplier), UAnimGenAutoEncoderTrainingSettings::StaticClass());
 		DataCategory.AddProperty(GET_MEMBER_NAME_CHECKED(UAnimGenAutoEncoderTrainingSettings, FrameAttributes), UAnimGenAutoEncoderTrainingSettings::StaticClass());
 
 		// Put the frame ranges returned by the query next
@@ -1058,9 +1060,12 @@ namespace UE::AnimGen::Editor
 						PoseVectorWeights.View,
 						TrainingDatabase->GetSkeleton()->GetReferenceSkeleton().GetRefBonePose(),
 						AutoEncoder->GetBoneParents(),
+						AutoEncoder->AutoEncodedRequiredBoneIndices,
 						AutoEncoder->AutoEncodedBoneLocationIndices,
 						AutoEncoder->AutoEncodedBoneRotationIndices,
-						AutoEncoder->AutoEncodedBoneScaleIndices);
+						AutoEncoder->AutoEncodedBoneScaleIndices,
+						AutoEncoder->TrainingSettings->RootWeightMultiplier,
+						AutoEncoder->TrainingSettings->BaseWeightMultiplier);
 				}
 
 				// Controls
